[StormSync] Moved plugin

#virtualized

diff --git a/Engine/Plugins/Experimental/StormSync/Source/StormSyncCore/Private/StormSyncCoreLog.cpp b/Engine/Plugins/Experimental/StormSync/Source/StormSyncCore/Private/StormSyncCoreLog.cpp
new file mode 100644
--- /dev/null
+++ b/Engine/Plugins/Experimental/StormSync/Source/StormSyncCore/Private/StormSyncCoreLog.cpp
@@ -0,0 +1,5 @@
+// Copyright Epic Games, Inc. All Rights Reserved.
+
+#include "StormSyncCoreLog.h"
+
+DEFINE_LOG_CATEGORY(LogStormSyncCore);
diff --git a/Engine/Plugins/Experimental/StormSync/Source/StormSyncCore/Private/StormSyncCoreModule.cpp b/Engine/Plugins/Experimental/StormSync/Source/StormSyncCore/Private/StormSyncCoreModule.cpp
new file mode 100644
--- /dev/null
+++ b/Engine/Plugins/Experimental/StormSync/Source/StormSyncCore/Private/StormSyncCoreModule.cpp
@@ -0,0 +1,13 @@
+// Copyright Epic Games, Inc. All Rights Reserved.
+
+#include "StormSyncCoreModule.h"
+
+void FStormSyncCoreModule::StartupModule()
+{
+}
+
+void FStormSyncCoreModule::ShutdownModule()
+{
+}
+
+IMPLEMENT_MODULE(FStormSyncCoreModule, StormSyncCore)
diff --git a/Engine/Plugins/Experimental/StormSync/Source/StormSyncCore/Private/StormSyncCoreSettings.cpp b/Engine/Plugins/Experimental/StormSync/Source/StormSyncCore/Private/StormSyncCoreSettings.cpp
new file mode 100644
--- /dev/null
+++ b/Engine/Plugins/Experimental/StormSync/Source/StormSyncCore/Private/StormSyncCoreSettings.cpp
@@ -0,0 +1,22 @@
+// Copyright Epic Games, Inc. All Rights Reserved.
+
+#include "StormSyncCoreSettings.h"
+
+UStormSyncCoreSettings::UStormSyncCoreSettings()
+{
+	IgnoredPackagesInternal.Add(TEXT("/Engine"));
+	IgnoredPackagesInternal.Add(TEXT("/Script"));
+	ExportDefaultNameFormatString = TEXT("%Y_%m_%d_%H%M%S");
+}
+
+FName UStormSyncCoreSettings::GetCategoryName() const
+{
+	return PluginCategoryName;
+}
+
+#if WITH_EDITOR
+FText UStormSyncCoreSettings::GetSectionText() const
+{
+	return NSLOCTEXT("StormSyncCorePlugin", "StormSyncCoreSettingsSection", "Core Settings");
+}
+#endif
diff --git a/Engine/Plugins/Experimental/StormSync/Source/StormSyncCore/Public/StormSyncCoreLog.h b/Engine/Plugins/Experimental/StormSync/Source/StormSyncCore/Public/StormSyncCoreLog.h
new file mode 100644
--- /dev/null
+++ b/Engine/Plugins/Experimental/StormSync/Source/StormSyncCore/Public/StormSyncCoreLog.h
@@ -0,0 +1,12 @@
+// Copyright Epic Games, Inc. All Rights Reserved.
+
+#pragma once
+
+#include "CoreMinimal.h"
+
+STORMSYNCCORE_API DECLARE_LOG_CATEGORY_EXTERN(LogStormSyncCore, Display, All);
+
+#define STORM_SYNC_LOG(Verbosity, Format, ...) \
+{ \
+    UE_LOG(LogStormSyncCore, Verbosity, Format, ##__VA_ARGS__); \
+}
diff --git a/Engine/Plugins/Experimental/StormSync/Source/StormSyncCore/Public/StormSyncCoreModule.h b/Engine/Plugins/Experimental/StormSync/Source/StormSyncCore/Public/StormSyncCoreModule.h
new file mode 100644
--- /dev/null
+++ b/Engine/Plugins/Experimental/StormSync/Source/StormSyncCore/Public/StormSyncCoreModule.h
@@ -0,0 +1,17 @@
+// Copyright Epic Games, Inc. All Rights Reserved.
+
+
+#pragma once
+
+#include "CoreMinimal.h"
+#include "Modules/ModuleManager.h"
+
+/** Main entry point and implementation of StormSync Core Runtime module. */
+class FStormSyncCoreModule : public IModuleInterface
+{
+public:
+	//~ Begin IModuleInterface interface
+	virtual void StartupModule() override;
+	virtual void ShutdownModule() override;
+	//~ End IModuleInterface interface
+};
diff --git a/Engine/Plugins/Experimental/StormSync/Source/StormSyncCore/Public/StormSyncImportTypes.h b/Engine/Plugins/Experimental/StormSync/Source/StormSyncCore/Public/StormSyncImportTypes.h
new file mode 100644
--- /dev/null
+++ b/Engine/Plugins/Experimental/StormSync/Source/StormSyncCore/Public/StormSyncImportTypes.h
@@ -0,0 +1,43 @@
+// Copyright Epic Games, Inc. All Rights Reserved.
+
+#pragma once
+
+#include "CoreMinimal.h"
+#include "StormSyncCommonTypes.h"
+#include "StormSyncPackageDescriptor.h"
+#include "StormSyncImportTypes.generated.h"
+
+/**
+ * Data holder during an extraction process.
+ *
+ * Used to compute a list of files that needs extraction.
+ */
+USTRUCT()
+struct STORMSYNCCORE_API FStormSyncImportFileInfo
+{
+	GENERATED_BODY()
+	
+	/** Holds info such as PackageName, FileSize, Timestamp and hash*/
+	FStormSyncFileDependency FileDependency;
+
+	/** The associated fully qualified path of the .uasset */
+	FString DestFilepath;
+
+	/** Holds a shared ptr to the serialized buffer */
+	FStormSyncBufferPtr FileBuffer = nullptr;
+
+	/** UI-exposed reason of the import */
+	FText ImportReason;
+	
+	/** UI-exposed extended info for reason of the import */
+	FText ImportReasonTooltip;
+
+	FStormSyncImportFileInfo() = default;
+
+	FStormSyncImportFileInfo(const FStormSyncFileDependency& InFileDependency, const FString& InDestFilepath, const FStormSyncBufferPtr& InFileBuffer)
+		: FileDependency(InFileDependency)
+		, DestFilepath(InDestFilepath)
+		, FileBuffer(InFileBuffer)
+	{
+	}
+};
diff --git a/Engine/Plugins/Experimental/StormSync/Source/StormSyncDrives/Private/StormSyncDrivesModule.cpp b/Engine/Plugins/Experimental/StormSync/Source/StormSyncDrives/Private/StormSyncDrivesModule.cpp
new file mode 100644
--- /dev/null
+++ b/Engine/Plugins/Experimental/StormSync/Source/StormSyncDrives/Private/StormSyncDrivesModule.cpp
@@ -0,0 +1,283 @@
+// Copyright Epic Games, Inc. All Rights Reserved.
+
+#include "StormSyncDrivesModule.h"
+
+#include "Engine/DeveloperSettings.h"
+#include "Misc/CoreDelegates.h"
+#include "Misc/PackageName.h"
+#include "StormSyncDrivesLog.h"
+#include "StormSyncDrivesSettings.h"
+#include "StormSyncDrivesUtils.h"
+
+#if WITH_EDITOR
+#include "MessageLogModule.h"
+#include "Logging/MessageLog.h"
+#endif
+
+#define LOCTEXT_NAMESPACE "FStormSyncDrivesModule"
+
+void FStormSyncDrivesModule::StartupModule()
+{
+#if WITH_EDITOR
+	// Create a message log for the validation to use
+	FMessageLogModule& MessageLogModule = FModuleManager::LoadModuleChecked<FMessageLogModule>("MessageLog");
+	FMessageLogInitializationOptions InitOptions;
+	InitOptions.bShowPages = true;
+	MessageLogModule.RegisterLogListing(LogName, LOCTEXT("StormSyncLogLabel", "Storm Sync Drives"), InitOptions);
+	LogListing = MessageLogModule.GetLogListing(LogName);
+#endif
+
+	FCoreDelegates::OnFEngineLoopInitComplete.AddRaw(this, &FStormSyncDrivesModule::OnEngineLoopInitComplete);
+
+#if WITH_EDITOR
+	GetMutableDefault<UStormSyncDrivesSettings>()->OnSettingChanged().AddRaw(this, &FStormSyncDrivesModule::OnSettingsChanged);
+#endif
+}
+
+void FStormSyncDrivesModule::ShutdownModule()
+{
+#if WITH_EDITOR
+	if (FModuleManager::Get().IsModuleLoaded("MessageLog"))
+	{
+		// unregister message log
+		FMessageLogModule& MessageLogModule = FModuleManager::GetModuleChecked<FMessageLogModule>("MessageLog");
+		MessageLogModule.UnregisterLogListing(LogName);
+	}
+#endif
+
+	FCoreDelegates::OnFEngineLoopInitComplete.RemoveAll(this);
+
+#if WITH_EDITOR
+	if (UObjectInitialized())
+	{
+		GetMutableDefault<UStormSyncDrivesSettings>()->OnSettingChanged().RemoveAll(this);
+	}
+#endif
+
+	MountedDrives.Empty();
+}
+
+bool FStormSyncDrivesModule::RegisterMountPoint(const FStormSyncMountPointConfig& InMountPoint, FText& ErrorText)
+{
+	STORM_SYNC_DRIVES_LOG(Display, TEXT("FStormSyncDrivesModule::RegisterMountPoint ... MountPoint: %s, Path: %s"), *InMountPoint.MountPoint, *InMountPoint.MountDirectory.Path)
+
+	// Validate first configuration of the mount point
+	FText ValidationError;
+	if (!FStormSyncDrivesUtils::ValidateMountPoint(InMountPoint, ValidationError))
+	{
+		ErrorText = FText::Format(
+			LOCTEXT("ValidationError_MountPoint", "Validation failed for MountPoint (MountPoint: {0}, MountDirectory: {1}) - {2}"),
+			FText::FromString(InMountPoint.MountPoint),
+			FText::FromString(InMountPoint.MountDirectory.Path),
+			ValidationError
+		);
+		AddMessageError(ErrorText);
+		return false;
+	}
+
+	// Append a trailing slash so that it is consistent with FLongPackagePathsSingleton::ContentPathToRoot entries and avoid a crash
+	// with assets and reimport path when right clicking on them
+
+	// ValidateMountPoint should ensure no trailing slash is allowed in user config
+	const FString MountPoint = InMountPoint.MountPoint + TEXT("/");
+
+	// Check if root path is already mounted
+	if (FPackageName::MountPointExists(MountPoint))
+	{
+		ErrorText = FText::Format(
+			LOCTEXT("ValidationError_MountPointExists", "Mount Point {0} already exists"),
+			FText::FromString(MountPoint)
+		);
+		AddMessageError(ErrorText);
+		return false;
+	}
+
+	STORM_SYNC_DRIVES_LOG(Display, TEXT("FStormSyncDrivesModule::RegisterMountPoint ... Try register %s to %s"), *MountPoint, *InMountPoint.MountDirectory.Path)
+	FPackageName::RegisterMountPoint(MountPoint, InMountPoint.MountDirectory.Path);
+	return true;
+}
+
+bool FStormSyncDrivesModule::UnregisterMountPoint(const FStormSyncMountPointConfig& InMountPoint, FText& ErrorText)
+{
+	const FString ContentPath = InMountPoint.MountDirectory.Path;
+
+	// Mount points are always added with a trailing slash (eg. /Foo => /Foo/)
+	const bool bHasLeadingSlash = InMountPoint.MountPoint.EndsWith(TEXT("/"));
+	const FString RootPath = bHasLeadingSlash ? InMountPoint.MountPoint : InMountPoint.MountPoint + TEXT("/");
+
+	if (!FPackageName::MountPointExists(RootPath))
+	{
+		ErrorText = FText::Format(
+			LOCTEXT("ValidationError_MountPointUnexisting", "Mount Point {0} does not exist"),
+			FText::FromString(RootPath)
+		);
+		return false;
+	}
+
+	FPackageName::UnRegisterMountPoint(RootPath, ContentPath);
+	return true;
+}
+
+void FStormSyncDrivesModule::OnEngineLoopInitComplete()
+{
+	STORM_SYNC_DRIVES_LOG(Display, TEXT("FStormSyncDrivesModule::OnEngineLoopInitComplete ..."))
+	STORM_SYNC_DRIVES_LOG(Display, TEXT("\t Mounting Drives based on config"))
+
+	ResetMountedDrivesFromSettings(GetDefault<UStormSyncDrivesSettings>());
+}
+
+void FStormSyncDrivesModule::OnSettingsChanged(UObject* InSettings, FPropertyChangedEvent& InPropertyChangedEvent)
+{
+	const FName PropertyName = InPropertyChangedEvent.GetPropertyName();
+	const FName MemberPropertyName = (InPropertyChangedEvent.MemberProperty != nullptr) ? InPropertyChangedEvent.MemberProperty->GetFName() : NAME_None;
+	const FName MountPointsPropertyName = GET_MEMBER_NAME_CHECKED(UStormSyncDrivesSettings, MountPoints);
+
+
+	if (PropertyName == MountPointsPropertyName || MemberPropertyName == MountPointsPropertyName)
+	{
+		ResetMountedDrivesFromSettings(CastChecked<UStormSyncDrivesSettings>(InSettings));
+	}
+}
+
+bool FStormSyncDrivesModule::ValidateMountPoint(const FStormSyncMountPointConfig& InMountPoint, int32 Index) const
+{
+	const FString RootPath = InMountPoint.MountPoint;
+	const FString Directory = InMountPoint.MountDirectory.Path;
+
+	FText ErrorText;
+	if (!FStormSyncDrivesUtils::ValidateMountPoint(InMountPoint, ErrorText))
+	{
+		const FText ErrorLog = FText::Format(
+			LOCTEXT("ValidationError_MountPoint", "Validation failed for MountPoint at index {0} (MountPoint: {1}, MountDirectory: {2}) - {3}"),
+			FText::AsNumber(Index),
+			FText::FromString(RootPath),
+			FText::FromString(Directory),
+			ErrorText
+		);
+		AddMessageError(ErrorLog);
+		STORM_SYNC_DRIVES_LOG(Error, TEXT("FStormSyncDrivesModule::ValidateMountPoint - %s"), *ErrorLog.ToString())
+		return false;
+	}
+
+	return true;
+}
+
+bool FStormSyncDrivesModule::ValidateMountPoints(const TArray<FStormSyncMountPointConfig>& InMountPoints) const
+{
+#if WITH_EDITOR
+	// Reset log listing and clear message on each validation run
+	if (LogListing.IsValid())
+	{
+		LogListing->ClearMessages();
+	}
+#endif
+
+	bool bResult = true;
+
+	// First check for mount point validity individually
+	int32 Index = 0;
+	for (const FStormSyncMountPointConfig& Entry : InMountPoints)
+	{
+		STORM_SYNC_DRIVES_LOG(Display, TEXT("Try validate entry %s %s"), *Entry.MountPoint, *Entry.MountDirectory.Path)
+
+		// First check for mount point validity individually
+		bResult &= ValidateMountPoint(Entry, Index);
+		Index++;
+	}
+
+	// Then check for duplicates
+	TArray<FText> ValidationErrors;
+	if (!FStormSyncDrivesUtils::ValidateNonDuplicates(InMountPoints, ValidationErrors))
+	{
+		bResult = false;
+
+		// Log any returned errors to message log
+		for (const FText& ValidationError : ValidationErrors)
+		{
+			AddMessageError(ValidationError);
+			STORM_SYNC_DRIVES_LOG(Error, TEXT("FStormSyncDrivesModule::ValidateMountPoints - %s"), *ValidationError.ToString())
+		}
+	}
+
+#if WITH_EDITOR
+	// Got some errors, notify user
+	if (!bResult)
+	{
+		LogListing->NotifyIfAnyMessages(LOCTEXT("Validation_Error", "Storm Sync: Error validating Mount Points settings"), EMessageSeverity::Error, true);
+	}
+#endif
+
+	return bResult;
+}
+
+void FStormSyncDrivesModule::ResetMountedDrivesFromSettings(const UStormSyncDrivesSettings* InSettings)
+{
+	STORM_SYNC_DRIVES_LOG(Display, TEXT("FStormSyncDrivesModule::ResetMountedDrives ..."))
+	check(InSettings);
+
+	if (ValidateMountPoints(InSettings->MountPoints))
+	{
+		UnregisterMountedDrives();
+		CacheMountedDrives(InSettings->MountPoints);
+		RegisterMountedDrives();
+	}
+}
+
+void FStormSyncDrivesModule::UnregisterMountedDrives()
+{
+	STORM_SYNC_DRIVES_LOG(Display, TEXT("FStormSyncDrivesModule::UnregisterMountedDrives ..."))
+	for (const TSharedPtr<FStormSyncMountPointConfig>& MountedDrive : MountedDrives)
+	{
+		if (!MountedDrive.IsValid())
+		{
+			continue;
+		}
+
+		FText ErrorText;
+		if (!UnregisterMountPoint(*MountedDrive.Get(), ErrorText))
+		{
+			STORM_SYNC_DRIVES_LOG(Error, TEXT("FStormSyncDrivesModule::UnregisterMountedDrives failed with Error: %s"), *ErrorText.ToString())
+		}
+	}
+}
+
+void FStormSyncDrivesModule::CacheMountedDrives(const TArray<FStormSyncMountPointConfig>& InMountPoints)
+{
+	STORM_SYNC_DRIVES_LOG(Display, TEXT("FStormSyncDrivesModule::CacheMountedDrives ..."))
+	MountedDrives.Reset();
+
+	for (const FStormSyncMountPointConfig& Entry : InMountPoints)
+	{
+		MountedDrives.Add(MakeShared<FStormSyncMountPointConfig>(Entry));
+	}
+}
+
+void FStormSyncDrivesModule::RegisterMountedDrives()
+{
+	STORM_SYNC_DRIVES_LOG(Display, TEXT("FStormSyncDrivesModule::RegisterMountedDrives ..."))
+	for (TSharedPtr<FStormSyncMountPointConfig> MountedDrive : MountedDrives)
+	{
+		if (!MountedDrive.IsValid())
+		{
+			continue;
+		}
+
+		FText ErrorText;
+		if (!RegisterMountPoint(*MountedDrive.Get(), ErrorText))
+		{
+			STORM_SYNC_DRIVES_LOG(Error, TEXT("FStormSyncDrivesModule::RegisterMountedDrives failed with Error: %s"), *ErrorText.ToString())
+		}
+	}
+}
+
+void FStormSyncDrivesModule::AddMessageError(const FText& Text)
+{
+#if WITH_EDITOR
+	FMessageLog MessageLog(LogName);
+	MessageLog.Error(Text);
+#endif
+}
+
+#undef LOCTEXT_NAMESPACE
+
+IMPLEMENT_MODULE(FStormSyncDrivesModule, StormSyncDrives)
diff --git a/Engine/Plugins/Experimental/StormSync/Source/StormSyncDrives/Private/StormSyncDrivesSettings.cpp b/Engine/Plugins/Experimental/StormSync/Source/StormSyncDrives/Private/StormSyncDrivesSettings.cpp
new file mode 100644
--- /dev/null
+++ b/Engine/Plugins/Experimental/StormSync/Source/StormSyncDrives/Private/StormSyncDrivesSettings.cpp
@@ -0,0 +1,20 @@
+// Copyright Epic Games, Inc. All Rights Reserved.
+
+
+#include "StormSyncDrivesSettings.h"
+
+UStormSyncDrivesSettings::UStormSyncDrivesSettings()
+{
+}
+
+FName UStormSyncDrivesSettings::GetCategoryName() const
+{
+	return PluginCategoryName;
+}
+
+#if WITH_EDITOR
+FText UStormSyncDrivesSettings::GetSectionText() const
+{
+	return NSLOCTEXT("StormSyncDrives", "StormSyncDrivesSettingsSection", "Mount Points Settings");
+}
+#endif
diff --git a/Engine/Plugins/Experimental/StormSync/Source/StormSyncDrives/Public/StormSyncDrivesLog.h b/Engine/Plugins/Experimental/StormSync/Source/StormSyncDrives/Public/StormSyncDrivesLog.h
new file mode 100644
--- /dev/null
+++ b/Engine/Plugins/Experimental/StormSync/Source/StormSyncDrives/Public/StormSyncDrivesLog.h
@@ -0,0 +1,12 @@
+// Copyright Epic Games, Inc. All Rights Reserved.
+
+#pragma once
+
+#include "CoreMinimal.h"
+
+STORMSYNCDRIVES_API DECLARE_LOG_CATEGORY_EXTERN(LogStormSyncDrives, Display, All);
+
+#define STORM_SYNC_DRIVES_LOG(Verbosity, Format, ...) \
+{ \
+    UE_LOG(LogStormSyncDrives, Verbosity, Format, ##__VA_ARGS__); \
+}
diff --git a/Engine/Plugins/Experimental/StormSync/Source/StormSyncEditor/Public/StormSyncEditorLog.h b/Engine/Plugins/Experimental/StormSync/Source/StormSyncEditor/Public/StormSyncEditorLog.h
new file mode 100644
--- /dev/null
+++ b/Engine/Plugins/Experimental/StormSync/Source/StormSyncEditor/Public/StormSyncEditorLog.h
@@ -0,0 +1,12 @@
+// Copyright Epic Games, Inc. All Rights Reserved.
+
+#pragma once
+
+#include "CoreMinimal.h"
+
+STORMSYNCEDITOR_API DECLARE_LOG_CATEGORY_EXTERN(LogStormSyncEditor, Display, All);
+
+#define STORM_SYNC_EDITOR_LOG(Verbosity, Format, ...) \
+{ \
+	UE_LOG(LogStormSyncEditor, Verbosity, Format, ##__VA_ARGS__); \
+}
diff --git a/Engine/Plugins/Experimental/StormSync/Source/StormSyncImport/Private/StormSyncImport.cpp b/Engine/Plugins/Experimental/StormSync/Source/StormSyncImport/Private/StormSyncImport.cpp
new file mode 100644
--- /dev/null
+++ b/Engine/Plugins/Experimental/StormSync/Source/StormSyncImport/Private/StormSyncImport.cpp
@@ -0,0 +1,21 @@
+// Copyright Epic Games, Inc. All Rights Reserved.
+
+#include "StormSyncImport.h"
+
+#include "StormSyncImportLog.h"
+
+#define LOCTEXT_NAMESPACE "FStormSyncImportModule"
+
+void FStormSyncImportModule::StartupModule()
+{
+	STORM_SYNC_IMPORT_LOG(Verbose, TEXT("Started StormSyncImport module ..."))
+}
+
+void FStormSyncImportModule::ShutdownModule()
+{
+    
+}
+
+#undef LOCTEXT_NAMESPACE
+    
+IMPLEMENT_MODULE(FStormSyncImportModule, StormSyncImport)
\ No newline at end of file
diff --git a/Engine/Plugins/Experimental/StormSync/Source/StormSyncImport/Private/StormSyncImport.h b/Engine/Plugins/Experimental/StormSync/Source/StormSyncImport/Private/StormSyncImport.h
new file mode 100644
--- /dev/null
+++ b/Engine/Plugins/Experimental/StormSync/Source/StormSyncImport/Private/StormSyncImport.h
@@ -0,0 +1,14 @@
+// Copyright Epic Games, Inc. All Rights Reserved.
+
+#pragma once
+
+#include "CoreMinimal.h"
+#include "IStormSyncImportModule.h"
+
+/** Storm Sync Import module implementation */
+class FStormSyncImportModule : public IStormSyncImportModule
+{
+public:
+    virtual void StartupModule() override;
+    virtual void ShutdownModule() override;
+};
diff --git a/Engine/Plugins/Experimental/StormSync/Source/StormSyncImport/Private/StormSyncImportLog.cpp b/Engine/Plugins/Experimental/StormSync/Source/StormSyncImport/Private/StormSyncImportLog.cpp
new file mode 100644
--- /dev/null
+++ b/Engine/Plugins/Experimental/StormSync/Source/StormSyncImport/Private/StormSyncImportLog.cpp
@@ -0,0 +1,6 @@
+// Copyright Epic Games, Inc. All Rights Reserved.
+
+#include "StormSyncImportLog.h"
+
+DEFINE_LOG_CATEGORY(LogStormSyncImport);
+
diff --git a/Engine/Plugins/Experimental/StormSync/Source/StormSyncImport/Private/Subsystems/StormSyncImportSubsystem.cpp b/Engine/Plugins/Experimental/StormSync/Source/StormSyncImport/Private/Subsystems/StormSyncImportSubsystem.cpp
new file mode 100644
--- /dev/null
+++ b/Engine/Plugins/Experimental/StormSync/Source/StormSyncImport/Private/Subsystems/StormSyncImportSubsystem.cpp
@@ -0,0 +1,974 @@
+// Copyright Epic Games, Inc. All Rights Reserved.
+
+#include "Subsystems/StormSyncImportSubsystem.h"
+
+#include "AssetRegistry/AssetData.h"
+#include "AssetRegistry/AssetRegistryModule.h"
+#include "Engine/Engine.h"
+#include "HAL/PlatformFileManager.h"
+#include "Logging/MessageLog.h"
+#include "Misc/FileHelper.h"
+#include "Misc/PackageName.h"
+#include "Misc/Paths.h"
+#include "Misc/ScopedSlowTask.h"
+#include "RenderingThread.h"
+#include "StormSyncCoreDelegates.h"
+#include "StormSyncCoreSettings.h"
+#include "StormSyncCoreUtils.h"
+#include "StormSyncImportLog.h"
+#include "StormSyncImportTypes.h"
+#include "StormSyncTransportSettings.h"
+#include "Tasks/IStormSyncImportTask.h"
+#include "TimerManager.h"
+#include "UObject/LinkerLoad.h"
+#include "UObject/Package.h"
+
+#if WITH_EDITOR
+#include "Algo/Transform.h"
+#include "Editor.h"
+#include "FileHelpers.h"
+#include "ISourceControlModule.h"
+#include "ISourceControlProvider.h"
+#include "MessageLogModule.h"
+#include "ObjectTools.h"
+#include "PackageTools.h"
+#include "Settings/EditorLoadingSavingSettings.h"
+#include "Slate/SStormSyncImportWizard.h"
+#include "SourceControlHelpers.h"
+#include "Subsystems/AssetEditorSubsystem.h"
+#include "Toolkits/ToolkitManager.h"
+#endif
+
+#define LOCTEXT_NAMESPACE "StormSyncImportSubsystem"
+
+void UStormSyncImportSubsystem::Initialize(FSubsystemCollectionBase& Collection)
+{
+	STORM_SYNC_IMPORT_LOG(Verbose, TEXT("UStormSyncImportSubsystem::Initialize (World: %s)"), *GetNameSafe(GetWorld()))
+
+#if WITH_EDITOR
+	// Create a message log for the asset tools to use
+	FMessageLogModule& MessageLogModule = FModuleManager::LoadModuleChecked<FMessageLogModule>("MessageLog");
+	FMessageLogInitializationOptions InitOptions;
+	InitOptions.bShowPages = true;
+	MessageLogModule.RegisterLogListing(LogName, LOCTEXT("StormSyncLogLabel", "Storm Sync Editor"), InitOptions);
+#endif
+}
+
+void UStormSyncImportSubsystem::Deinitialize()
+{
+	STORM_SYNC_IMPORT_LOG(Verbose, TEXT("UStormSyncImportSubsystem::Deinitialize (World: %s)"), *GetNameSafe(GetWorld()))
+	
+#if WITH_EDITOR
+	if (FModuleManager::Get().IsModuleLoaded("MessageLog"))
+	{
+		// unregister message log
+		FMessageLogModule& MessageLogModule = FModuleManager::GetModuleChecked<FMessageLogModule>("MessageLog");
+		MessageLogModule.UnregisterLogListing(LogName);
+	}
+#endif
+
+	// Cleanup delegates
+	FStormSyncCoreDelegates::OnRequestImportBuffer.RemoveAll(this);
+	FStormSyncCoreDelegates::OnRequestImportFile.RemoveAll(this);
+}
+
+UStormSyncImportSubsystem& UStormSyncImportSubsystem::Get()
+{
+	return *GEngine->GetEngineSubsystem<UStormSyncImportSubsystem>();
+}
+
+bool UStormSyncImportSubsystem::EnqueueImportTask(const TSharedPtr<IStormSyncImportSubsystemTask>& InImportTask, const UWorld* InWorld)
+{
+	if (!InWorld)
+	{
+		STORM_SYNC_IMPORT_LOG(Error, TEXT("UStormSyncImportSubsystem::EnqueueImportTask failed because of invalid world (World: %s)"), *GetNameSafe(InWorld))
+		return false;
+	}
+
+	// We specifically check for existing pending tasks because of the following scenario:
+	//
+	// Now that UStormSyncImportWorldSubsystem exists to handle import in -game mode, we need to account for possibly 2 instance of this subsystem at once,
+	// especially with PIE (Play in Editor). The currently loaded map in editor viewport will have its own World subsystem listening for imports, and if a PIE
+	// session is running (possibly more when simulating clients), each session will have its own World subsystem instance running as well, in addition to the
+	// editor one.
+	//
+	// Only allowing one task at a time may be a bit naive / brute force though
+
+	if (!PendingTasks.IsEmpty())
+	{
+		STORM_SYNC_IMPORT_LOG(Error, TEXT("UStormSyncImportSubsystem::EnqueueImportTask failed because of existing pending tasks (World: %s)"), *GetNameSafe(InWorld))
+		return false;
+	}
+
+	const bool bIsQueued = PendingTasks.Enqueue(InImportTask);
+
+	if (!NextTickHandler.IsValid())
+	{
+		NextTickHandler = FTSTicker::GetCoreTicker().AddTicker(FTickerDelegate::CreateUObject(this, &UStormSyncImportSubsystem::HandleNextTick));
+	}
+	
+	return bIsQueued;
+}
+
+bool UStormSyncImportSubsystem::PerformFileImport(const FString& InFilename)
+{
+	STORM_SYNC_IMPORT_LOG(Display, TEXT("FStormSyncEditorUtils::PerformImport for %s"), *InFilename);
+	TArray<uint8> Data;
+
+	const FStormSyncBufferPtr Buffer = MakeShareable(new TArray<uint8>);
+	if (!FFileHelper::LoadFileToArray(*Buffer.Get(), *InFilename))
+	{
+		STORM_SYNC_IMPORT_LOG(Error, TEXT("Failed to load file '%s' to array"), *InFilename);
+		return false;
+	}
+
+	// Dummy package descriptor with name set to imported filename (Note: Should consider serialize package descriptor along buffer header)
+	FStormSyncPackageDescriptor PackageDescriptor;
+	PackageDescriptor.Name = FPaths::GetBaseFilename(InFilename);
+
+	constexpr bool bShowWizard = true;
+	return PerformImport(PackageDescriptor, Buffer, bShowWizard);
+}
+
+bool UStormSyncImportSubsystem::PerformBufferImport(const FStormSyncPackageDescriptor& InPackageDescriptor, const FStormSyncBufferPtr& InBuffer)
+{
+	if (!InBuffer.IsValid())
+	{
+		STORM_SYNC_IMPORT_LOG(Error, TEXT("UStormSyncImportSubsystem::PerformBufferImport failed to import from invalid buffer"))
+		return false;
+	}
+
+	const UStormSyncTransportSettings* Settings = GetDefault<UStormSyncTransportSettings>();
+	check(Settings);
+	
+	const bool bTcpDryRun = Settings->IsTcpDryRun();
+	const bool bShowWizard = Settings->ShouldShowImportWizard();
+	return PerformImport(InPackageDescriptor, InBuffer, bShowWizard, bTcpDryRun);
+}
+
+bool UStormSyncImportSubsystem::PerformImport(const FStormSyncPackageDescriptor& InPackageDescriptor, const FStormSyncBufferPtr& InBuffer, bool bShowWizard, bool bDryRun)
+{
+	if (!InBuffer.IsValid())
+	{
+		STORM_SYNC_IMPORT_LOG(Error, TEXT("UStormSyncImportSubsystem::PerformBufferImport failed to import from invalid buffer"))
+		return false;
+	}
+
+	// Package descriptor local copy
+	FStormSyncPackageDescriptor PackageDescriptor = InPackageDescriptor;
+	
+	float InitialAmountOfWork = 1.f;
+	FScopedSlowTask SlowTask(InitialAmountOfWork, LOCTEXT("ImportPackages_Extracting", "Figuring out file states..."));
+	SlowTask.MakeDialog();
+
+	FStormSyncCoreExtractArgs ExtractArgs;
+
+	// The list of files with changes detected (either size of hash changed)
+	TArray<FStormSyncImportFileInfo> FilesToImport;
+
+	// The list of all files included in the buffer for filename
+	TArray<FStormSyncImportFileInfo> BufferFiles;
+
+	ExtractArgs.OnPakPreExtract.BindLambda([&SlowTask, PackageDescriptor, bDryRun](const int32 FileCount)
+	{
+		SlowTask.TotalAmountOfWork = FileCount;
+		STORM_SYNC_IMPORT_LOG(Display, TEXT("\tOnPakPreExtract - File Count: %d"), FileCount)
+		
+		if (!bDryRun)
+		{
+			UStormSyncImportSubsystem::Get().HandlePakPreExtract(PackageDescriptor, FileCount);
+		}
+	});
+	
+	ExtractArgs.OnPakPostExtract.BindLambda([PackageDescriptor, bShowWizard, bDryRun, &FilesToImport, &BufferFiles](const int32 FileCount)
+	{
+		if (bDryRun)
+		{
+			STORM_SYNC_IMPORT_LOG(Display, TEXT("\tOnPakPostExtract - File Count: %d"), FileCount)
+			return;
+		}
+
+#if WITH_EDITOR
+		if (bShowWizard)
+		{
+			// Create file import dialog. This is a modal dialog so it will return only after user selection
+			const TSharedRef<SStormSyncImportWizard> Wizard = SStormSyncImportWizard::CreateWizard(FilesToImport, BufferFiles);
+
+			// Early out if user canceled the operation
+			if (!Wizard->ShouldImport())
+			{
+				UStormSyncImportSubsystem::Get().HandlePakPostExtract(PackageDescriptor, 0);
+				return;
+			}
+		}
+#endif
+
+		// Begin import process
+		{
+			TArray<const FStormSyncImportFileInfo*> ExistingFiles;
+			TArray<const FStormSyncImportFileInfo*> NewFiles;
+			SplitImportFiles(FilesToImport, ExistingFiles, NewFiles);
+
+			FScopedSlowTask SlowTaskExtract(FilesToImport.Num() + ExistingFiles.Num() + NewFiles.Num()
+				, LOCTEXT("ImportPackages_Extracting", "Importing files to project..."));
+
+			// Step 1: Prompt for Checkout (for existing assets)
+			SlowTaskExtract.EnterProgressFrame(ExistingFiles.Num());
+			UStormSyncImportSubsystem::Get().HandleExistingAssets(ExistingFiles, bShowWizard);
+
+			// Step 2: Extract
+			for (const FStormSyncImportFileInfo& FileInfo : FilesToImport)
+			{
+				SlowTaskExtract.EnterProgressFrame();
+	
+				// Handle extraction, this is where assets are imported to project (closing asset editor if opened, deleting the asset and then writing it back)
+				UStormSyncImportSubsystem::Get().HandlePakAssetExtract(FileInfo.FileDependency, FileInfo.DestFilepath, FileInfo.FileBuffer);
+			}
+
+			// Step 3: Mark for Add (for new assets)
+			SlowTaskExtract.EnterProgressFrame(NewFiles.Num());
+			UStormSyncImportSubsystem::Get().HandleNewAssets(NewFiles, bShowWizard);
+		}
+
+		// Extract completion
+		UStormSyncImportSubsystem::Get().HandlePakPostExtract(PackageDescriptor, FilesToImport.Num());
+	});
+	
+	ExtractArgs.OnFileExtract.BindLambda([&SlowTask, bDryRun, &FilesToImport, &BufferFiles](const FStormSyncFileDependency& FileDependency, const FString DestFilepath, const FStormSyncBufferPtr& FileBuffer)
+	{
+		SlowTask.EnterProgressFrame();
+		
+		if (bDryRun)
+		{
+			STORM_SYNC_IMPORT_LOG(Display, TEXT("\tOnFileExtract"))
+			STORM_SYNC_IMPORT_LOG(Display, TEXT("\t\tPackageName: %s"), *FileDependency.PackageName.ToString())
+			STORM_SYNC_IMPORT_LOG(Display, TEXT("\t\tDestFilepath: %s"), *DestFilepath)
+			STORM_SYNC_IMPORT_LOG(Display, TEXT("\t\tFileSize: %lld"), FileDependency.FileSize)
+			STORM_SYNC_IMPORT_LOG(Display, TEXT("\t\tFileBuffer valid: %s"), FileBuffer.IsValid() ? TEXT("valid") : TEXT("invalid"))
+			if (FileBuffer.IsValid())
+			{
+				STORM_SYNC_IMPORT_LOG(Display, TEXT("\t\tFileBuffer size: %d"), FileBuffer->Num())
+			}
+
+			return;
+		}
+		
+		// Figure out files that needs import and build list of FStormSyncImportFileInfo for Import Wizard widget
+		const FStormSyncImportFileInfo FileInfo = FStormSyncImportFileInfo(FileDependency, DestFilepath, FileBuffer);
+		BufferFiles.Add(FileInfo);
+	
+		UStormSyncImportSubsystem::Get().FillFilesToImport(FileInfo, DestFilepath, FilesToImport);
+	});
+
+	TArray<FText> Errors;
+	TMap<FString, FString> ExtractedPackages;
+	if (!FStormSyncCoreUtils::ExtractPakBuffer(*InBuffer.Get(), ExtractArgs, ExtractedPackages, Errors))
+	{
+		STORM_SYNC_IMPORT_LOG(Warning, TEXT("FStormSyncEditorUtils::PerformImport - Error extracting package ..."))
+		for (const FText& Error : Errors)
+		{
+			STORM_SYNC_IMPORT_LOG(Warning, TEXT("\t %s"), *Error.ToString())
+		}
+
+		return false;
+	}
+
+	return true;
+}
+
+bool UStormSyncImportSubsystem::HandleNextTick(float InDeltaTime)
+{
+	if (!PendingTasks.IsEmpty())
+	{
+		TSharedPtr<IStormSyncImportSubsystemTask> Task;
+		while (PendingTasks.Dequeue(Task))
+		{
+			if (Task.IsValid())
+			{
+				Task->Run();
+			}
+		}
+	}
+	NextTickHandler.Reset();
+	return false;
+}
+
+bool UStormSyncImportSubsystem::FillFilesToImport(const FStormSyncImportFileInfo& InFileInfo, const FString& InDestFilepath, TArray<FStormSyncImportFileInfo>& OutFilesToImport) const
+{
+	const FStormSyncFileDependency FileDependency = InFileInfo.FileDependency;
+
+	// Local copy
+	FStormSyncImportFileInfo FileInfo = InFileInfo;
+	
+	const int64 DestFileSize = IFileManager::Get().FileSize(*InDestFilepath);
+	// File not existing in local project, add it to files to import and early out
+	if (DestFileSize == INDEX_NONE)
+	{
+		STORM_SYNC_IMPORT_LOG(Verbose, TEXT("\t\tFile %s does not exist"), *InDestFilepath)
+		FString ShortPath = InDestFilepath;
+		ShortPath.RemoveFromStart(FPaths::ProjectContentDir());
+		FileInfo.ImportReason = FText::Format(LOCTEXT("FileInfo_MissingFile", "Missing file in local project ({0})"), FText::FromString(ShortPath));
+		FileInfo.ImportReasonTooltip = FText::Format(LOCTEXT("FileInfo_MissingFile", "File {0} does not exist locally"), FText::FromString(InDestFilepath));
+		OutFilesToImport.Add(FileInfo);
+		return true;
+	}
+
+	STORM_SYNC_IMPORT_LOG(Verbose, TEXT("\t\tFile %s exist"), *InDestFilepath)
+
+	const bool bSameSize = FileDependency.FileSize == DestFileSize;
+	STORM_SYNC_IMPORT_LOG(Verbose, TEXT("\t\t\tSame Size: %s (%lld vs %lld)"), bSameSize ? TEXT("true") : TEXT("false"), FileDependency.FileSize, DestFileSize)
+
+	// Check file hash from buffer against local file
+	const FMD5Hash ExistingFileMD5 = FMD5Hash::HashFile(*InDestFilepath);
+	const FString ExistingFileHash = LexToString(ExistingFileMD5);
+
+	const bool bSameHash = FileDependency.FileHash == ExistingFileHash;
+
+	// Mismatched File size or hash
+	if (!bSameSize)
+	{
+		FileInfo.ImportReason = FText::Format(
+			LOCTEXT("FileInfo_MismatchedSize_Tooltip", "Files have different sizes ({0} vs {1})"),
+			FText::FromString(FStormSyncCoreUtils::GetHumanReadableByteSize(DestFileSize)),
+			FText::FromString(FStormSyncCoreUtils::GetHumanReadableByteSize(FileDependency.FileSize))
+		);
+		FileInfo.ImportReasonTooltip = FText::Format(
+			LOCTEXT("FileInfo_MismatchedSize_Tooltip", "Mismatched file size ({0} vs {1})"),
+			FText::AsNumber(DestFileSize),
+			FText::AsNumber(FileDependency.FileSize)
+		);
+
+		OutFilesToImport.Add(FileInfo);
+		return true;
+	}
+	
+	if (!bSameHash)
+	{
+		FileInfo.ImportReason = LOCTEXT("FileInfo_MismatchedHash", "Files hash are not matching.");
+		FileInfo.ImportReasonTooltip = FText::Format(
+			LOCTEXT("FileInfo_MismatchedHash_Tooltip", "Mismatched file hash ({0} vs {1})"),
+			FText::FromString(ExistingFileHash),
+			FText::FromString(FileDependency.FileHash)
+		);
+
+		OutFilesToImport.Add(FileInfo);
+		return true;
+	}
+
+	return false;
+}
+
+void UStormSyncImportSubsystem::HandlePakPreExtract(const FStormSyncPackageDescriptor& InPackageDescriptor, const int32 FileCount)
+{
+	STORM_SYNC_IMPORT_LOG(Display, TEXT("UStormSyncImportSubsystem::HandlePakPreExtract - About to extract %d files for %s"), FileCount, *InPackageDescriptor.ToString());
+
+	// Reset closed assets to reopen
+	ClosedPackageNames.Empty();
+
+	// Reset extracted file reports
+	ExtractedFileReports.Empty();
+
+#if WITH_EDITOR
+	// Init a new page for message log to categorize each pak we receive
+	FMessageLog MessageLog(LogName);
+
+	const FText PageMessage = FText::Format(LOCTEXT("PakPreExtract_Incoming_Pak", "Handle incoming pak \"{0}\" ({1})"), FText::FromString(InPackageDescriptor.Name), FText::AsDateTime(FDateTime::Now()));
+	MessageLog.NewPage(PageMessage);
+	MessageLog.Info(FText::FromString(InPackageDescriptor.ToString()));
+#endif
+}
+
+void UStormSyncImportSubsystem::HandlePakPostExtract(const FStormSyncPackageDescriptor& InPackageDescriptor, const int32 FileCount) const
+{
+	STORM_SYNC_IMPORT_LOG(Display, TEXT("UStormSyncImportSubsystem::HandlePakPostExtract - Extracted %d files for %s"), FileCount, *InPackageDescriptor.ToString());
+	STORM_SYNC_IMPORT_LOG(Display, TEXT("UStormSyncImportSubsystem::HandlePakPostExtract - Should handle reopening of %d assets"), ClosedPackageNames.Num());
+	
+	if (GetDefault<UStormSyncCoreSettings>()->bEnableHotReloadPackages)
+	{
+		HotReloadPackages(ExtractedFileReports, false);
+	}
+	
+	OpenClosedEditors(ClosedPackageNames);
+
+	for (const FStormSyncEditorFileReport& FileReport : ExtractedFileReports)
+	{
+		if (FileReport.bSuccess)
+		{
+			FStormSyncCoreDelegates::OnPakAssetExtracted.Broadcast(FileReport.PackageName, FileReport.DestFilepath);
+		}
+	}
+	
+#if WITH_EDITOR
+	// Determine log severity and message based on our stored reports
+	EMessageSeverity::Type Severity = EMessageSeverity::Info;
+	FText LogMessage = FText::Format(LOCTEXT("PakPostExtract_Extract_Success", "Extracted {0} files successfully."), FText::AsNumber(FileCount));
+
+	const TArray<FStormSyncEditorFileReport> ErrorFileReports = ExtractedFileReports.FilterByPredicate([](const FStormSyncEditorFileReport& Report)
+	{
+		return Report.bSuccess == false;
+	});
+
+	if (!ErrorFileReports.IsEmpty())
+	{
+		Severity = EMessageSeverity::Error;
+		LogMessage = FText::Format(LOCTEXT("PakPostExtract_Extract_Failed", "Extracted {0} files. Some content could not be extracted."), FText::AsNumber(FileCount));
+	}
+
+	const FText LogHeading = FText::Format(LOCTEXT("PakPostExtract_Heading", "Received avalanche pak \"{0}\"."), FText::FromString(InPackageDescriptor.Name));
+	const FText NotifyMessage = FText::Format(LOCTEXT("PakPostExtract_Message_Notification", "{0}\n{1}"), LogHeading, LogMessage);
+	
+	FMessageLog MessageLog(LogName);
+	MessageLog.Notify(NotifyMessage, Severity, true);
+#endif
+}
+
+void UStormSyncImportSubsystem::HandleExistingAssets(TConstArrayView<const FStormSyncImportFileInfo*> InExistingFiles, bool bInShowPrompt)
+{
+#if WITH_EDITOR
+	STORM_SYNC_IMPORT_LOG(Display, TEXT("UStormSyncImportSubsystem::HandleExistingAssets - Package Count: %d"), InExistingFiles.Num());
+
+	TArray<UPackage*> ExistingPackages;
+
+	// Worst case: all files already exist in destination
+	ExistingPackages.Reserve(InExistingFiles.Num());
+
+	for (const FStormSyncImportFileInfo* FileInfo : InExistingFiles)
+	{
+		const FString PackageName = FileInfo->FileDependency.PackageName.ToString();
+
+		if (UPackage* const Package = LoadPackage(nullptr, *PackageName, LOAD_None))
+		{
+			STORM_SYNC_IMPORT_LOG(Verbose, TEXT("UStormSyncImportSubsystem::HandleExistingAssets - Existing Package found for: %s"), *PackageName);
+			ExistingPackages.Add(Package);
+		}
+		else
+		{
+			STORM_SYNC_IMPORT_LOG(Verbose, TEXT("UStormSyncImportSubsystem::HandleExistingAssets - Package not found for: %s. Not an Existing Asset"), *PackageName);
+		}
+	}
+
+	bool bSucceeded;
+
+	TArray<UPackage*> PackagesCheckedOut;
+	TArray<UPackage*> PackagesNotNeedingCheckout;
+
+	if (bInShowPrompt)
+	{
+		constexpr bool bCheckDirty = false;
+		constexpr bool bAllowSkip = false;
+		constexpr bool bPromptingAfterModify = false;
+
+		bSucceeded = FEditorFileUtils::PromptToCheckoutPackages(bCheckDirty
+			, ExistingPackages
+			, &PackagesCheckedOut
+			, &PackagesNotNeedingCheckout
+			, bPromptingAfterModify
+			, bAllowSkip);
+	}
+	else
+	{
+		constexpr bool bErrorIfAlreadyCheckedOut = false;
+		constexpr bool bConfirmPackageBranchCheckOutStatus = false;
+
+		ISourceControlProvider& SourceControlProvider = ISourceControlModule::Get().GetProvider();
+
+		if (SourceControlProvider.IsEnabled())
+		{
+			STORM_SYNC_IMPORT_LOG(Log, TEXT("UStormSyncImportSubsystem::HandleExistingAssets - Checking out %d packages with %s")
+				, ExistingPackages.Num()
+				, *SourceControlProvider.GetName().ToString());
+
+			bSucceeded = ECommandResult::Succeeded == FEditorFileUtils::CheckoutPackages(ExistingPackages
+				, &PackagesCheckedOut
+				, bErrorIfAlreadyCheckedOut
+				, bConfirmPackageBranchCheckOutStatus);
+		}
+		else
+		{
+			STORM_SYNC_IMPORT_LOG(Log, TEXT("UStormSyncImportSubsystem::HandleExistingAssets - No Source Control found. Making %d packages writeable")
+				, ExistingPackages.Num());
+
+			bSucceeded = UStormSyncImportSubsystem::MakePackagesWriteable(ExistingPackages) > 0;
+		}
+	}
+
+	if (!bSucceeded)
+	{
+		// TODO: at this point in time, the whole import could be cancelled and failed
+		STORM_SYNC_IMPORT_LOG(Error, TEXT("UStormSyncImportSubsystem::HandleExistingAssets - CheckoutPackages failed"));
+	}
+#endif
+}
+
+bool UStormSyncImportSubsystem::HandleNewAssets(TConstArrayView<const FStormSyncImportFileInfo*> InNewFiles, bool bInShowPrompt)
+{
+#if WITH_EDITOR
+	if (!GEditor || GIsAutomationTesting || IsRunningCommandlet())
+	{
+		return false;
+	}
+
+	if (!ISourceControlModule::Get().IsEnabled())
+	{
+		STORM_SYNC_IMPORT_LOG(Log, TEXT("UStormSyncImportSubsystem::HandleNewAssets - No Source Control found"));
+		return false;
+	}
+
+	if (!GetDefault<UEditorLoadingSavingSettings>()->bSCCAutoAddNewFiles)
+	{
+		STORM_SYNC_IMPORT_LOG(Log, TEXT("UStormSyncImportSubsystem::HandleNewAssets - Source Control Auto Add New Files is disabled."));
+		return false;
+	}
+
+	if (!ISourceControlModule::Get().GetProvider().IsAvailable())
+	{
+		STORM_SYNC_IMPORT_LOG(Log, TEXT("UStormSyncImportSubsystem::HandleNewAssets - Source Control Provider is not available."));
+		return false;
+	}
+
+	TArray<FString> FilesToAutoAdd;
+	FilesToAutoAdd.Reserve(InNewFiles.Num());
+
+	Algo::Transform(InNewFiles, FilesToAutoAdd,
+		[](const FStormSyncImportFileInfo* InFileInfo)
+		{
+			return InFileInfo->DestFilepath;
+		});
+
+	STORM_SYNC_IMPORT_LOG(Log, TEXT("UStormSyncImportSubsystem::HandleNewAssets - Found %d files to auto-add"), FilesToAutoAdd.Num());
+
+	const bool bSilent = !bInShowPrompt;
+	return USourceControlHelpers::CheckOutOrAddFiles(FilesToAutoAdd, bSilent);
+#endif
+	return false;
+}
+
+void UStormSyncImportSubsystem::HandlePakAssetExtract(const FStormSyncFileDependency& FileDependency, const FString& DestFilepath, const FStormSyncBufferPtr& FileBuffer)
+{
+	STORM_SYNC_IMPORT_LOG(Display, TEXT("UStormSyncImportSubsystem::HandlePakAssetExtract - Handle extracted package: %s"), *FileDependency.ToString());
+	STORM_SYNC_IMPORT_LOG(Display, TEXT("\tshould extract to %s"), *DestFilepath);
+	STORM_SYNC_IMPORT_LOG(Display, TEXT("\tValid buffer ? %s"), FileBuffer.IsValid() && FileBuffer->Num() != 0 ? TEXT("true") : TEXT("false"));
+
+	if (!FileBuffer.IsValid())
+	{
+		STORM_SYNC_IMPORT_LOG(Error, TEXT("UStormSyncImportSubsystem::HandlePakAssetExtract - Invalid file buffer"));
+		return;
+	}
+
+
+	FStormSyncEditorFileReport FileReport;
+	FileReport.PackageName = FileDependency.PackageName;
+	FileReport.DestFilepath = DestFilepath;
+
+	if (!FPaths::FileExists(DestFilepath))
+	{
+		if (WriteFile(DestFilepath, FileDependency.FileSize, FileBuffer->GetData()))
+		{
+			FileReport.bSuccess = true;
+		}
+		else
+		{
+			FileReport.bSuccess = false;
+		}
+
+		ExtractedFileReports.Add(FileReport);
+
+#if WITH_EDITOR
+		FMessageLog MessageLog(LogName);
+		if (FileReport.bSuccess)
+		{
+			MessageLog.Info(FText::Format(LOCTEXT("PakAssetExtract_SucessfulExtract", "Extracted {0} file to {1}."), FText::FromName(FileDependency.PackageName), FText::FromString(DestFilepath)));
+		}
+		else
+		{
+			MessageLog.Error(FText::Format(LOCTEXT("PakAssetExtract_FileWrite_Failed", "Failed to extract {0} file to {1}"), FText::FromName(FileDependency.PackageName), FText::FromString(DestFilepath)));
+		}
+#endif
+
+		return;
+	}
+
+
+	// If File exists locally, we should:
+	//
+	// 1. Close Editors
+	// 2. Delete file -- not necessary if using hot reload.
+	// 3. Write/Overwrite new file version from incoming pak
+	// 4. Hot Reload packages
+	// 5. Reopen closed editor
+
+	const FAssetRegistryModule& AssetRegistryModule = FModuleManager::Get().LoadModuleChecked<FAssetRegistryModule>(TEXT("AssetRegistry"));
+
+	// Close the asset editors if opened, we will attempt to reopen when we receive PostExtract event (handled in OnPakPostExtract)
+	TArray<FAssetData> AssetsToDelete;
+	AssetRegistryModule.Get().GetAssetsByPackageName(FileDependency.PackageName, AssetsToDelete);
+
+	TArray<FAssetData> ClosedAssets;
+	CloseEditors(AssetsToDelete, ClosedAssets);
+
+	// Store closed assets to reopen
+	for (FAssetData& AssetData : ClosedAssets)
+	{
+		ClosedPackageNames.AddUnique(AssetData.PackageName.ToString());
+	}
+
+	// Deleting the assets shouldn't be necessary when hot reload is enabled.
+	// TODO: We might still need to identify the assets that need to be deleted.
+	// See ConcertSyncClientUtil::PurgePackages.
+	if (!GetDefault<UStormSyncCoreSettings>()->bEnableHotReloadPackages)
+	{
+		if (!DeleteAssets(AssetsToDelete))
+		{
+			FileReport.bSuccess = false;
+			ExtractedFileReports.Add(FileReport);
+			return;
+		}
+	}
+	else
+	{
+		// They do this in FConcertClientPackageManager::SavePackageFile prior to writing data to file.
+		FlushPackageLoading(FileDependency.PackageName.ToString());
+	}
+
+	// Write the file to destination
+	if (WriteFile(DestFilepath, FileDependency.FileSize, FileBuffer->GetData()))
+	{
+		FileReport.bSuccess = true;
+	}
+	else
+	{
+		FileReport.bSuccess = false;
+	}
+
+#if WITH_EDITOR
+	FMessageLog MessageLog(LogName);
+	if (FileReport.bSuccess)
+	{
+		MessageLog.Info(FText::Format(LOCTEXT("PakAssetExtract_SucessfulExtract", "Extracted {0} file to {1}."), FText::FromName(FileDependency.PackageName), FText::FromString(DestFilepath)));
+	}
+	else
+	{
+		MessageLog.Error(FText::Format(LOCTEXT("PakAssetExtract_FileWrite_Failed", "Failed to extract {0} file to {1}"), FText::FromName(FileDependency.PackageName), FText::FromString(DestFilepath)));
+	}
+#endif
+
+	ExtractedFileReports.Add(FileReport);
+}
+
+void UStormSyncImportSubsystem::CloseEditors(const TArray<FAssetData>& InAssets, TArray<FAssetData>& OutClosedAssets)
+{
+#if WITH_EDITOR
+	for (const FAssetData& AssetData : InAssets)
+	{
+		const UObject* Asset = AssetData.FastGetAsset();
+		STORM_SYNC_IMPORT_LOG(Verbose, TEXT("\tClosing asset: %s (UObject: %s)"), *AssetData.GetFullName(), *GetNameSafe(Asset));
+
+		if (!Asset)
+		{
+			continue;
+		}
+
+		const TSharedPtr<IToolkit> AssetEditor = FToolkitManager::Get().FindEditorForAsset(Asset);
+		if (AssetEditor.IsValid() && IsAssetCurrentlyBeingEdited(AssetEditor, Asset))
+		{
+			FToolkitManager::Get().CloseToolkit(AssetEditor.ToSharedRef());
+
+			// Store assets to delete and reopen
+			OutClosedAssets.Add(AssetData);
+		}
+	}
+#endif
+}
+
+void UStormSyncImportSubsystem::OpenClosedEditors(const TArray<FString>& ClosedPackageNames)
+{
+	STORM_SYNC_IMPORT_LOG(Verbose, TEXT("UStormSyncImportSubsystem::OpenClosedEditors - ClosedPackageNames: %d"), ClosedPackageNames.Num())
+#if WITH_EDITOR
+	if (!GEditor)
+	{
+		return;
+	}
+	
+	STORM_SYNC_IMPORT_LOG(Verbose, TEXT("UStormSyncImportSubsystem::OpenClosedEditors - ClosedPackageNames: %d (WITH_EDITOR)"), ClosedPackageNames.Num())
+	TArray<UObject*> ObjectsToReopen;
+	for (FString ClosedPackageName : ClosedPackageNames)
+	{
+		if (UObject* Object = StaticLoadObject(UObject::StaticClass(), nullptr, *ClosedPackageName))
+		{
+			ObjectsToReopen.Add(Object);
+		}
+	}
+
+	if (UAssetEditorSubsystem* AssetEditorSubsystem = GEditor->GetEditorSubsystem<UAssetEditorSubsystem>())
+	{
+		AssetEditorSubsystem->OpenEditorForAssets(ObjectsToReopen);
+	}
+#endif
+}
+
+bool UStormSyncImportSubsystem::DeleteAssets(const TArray<FAssetData>& AssetsToDelete, const bool bShowConfirmation)
+{
+#if WITH_EDITOR
+	TArray<TWeakObjectPtr<UPackage>> PackageFilesToDelete;
+	TArray<UObject*> ObjectsToDelete;
+
+	for (const FAssetData& AssetData : AssetsToDelete)
+	{
+		UObject* ObjectToDelete = AssetData.GetAsset({ULevel::LoadAllExternalObjectsTag});
+		// Assets can be loaded even when their underlying type/class no longer exists...
+		if (ObjectToDelete != nullptr)
+		{
+			ObjectsToDelete.Add(ObjectToDelete);
+		}
+		else if (AssetData.IsUAsset())
+		{
+			// ... In this cases there is no underlying asset or type so remove the package itself directly after confirming it's valid to do so.
+			FString PackageFilename;
+			if (!FPackageName::DoesPackageExist(AssetData.PackageName.ToString(), &PackageFilename))
+			{
+				// Could not determine filename for package so we can not delete
+				ensureMsgf(false, TEXT("Could not determine filename for package %s so we can not delete"), *AssetData.PackageName.ToString());
+				continue;
+			}
+
+			if (UPackage* Package = FindPackage(nullptr, *AssetData.PackageName.ToString()))
+			{
+				PackageFilesToDelete.Add(Package);
+			}
+		}
+	}
+
+	int32 NumObjectsToDelete = ObjectsToDelete.Num();
+	if (NumObjectsToDelete > 0)
+	{
+		// First try with regular delete objects
+		NumObjectsToDelete = ObjectTools::DeleteObjects(ObjectsToDelete, bShowConfirmation);
+		if (NumObjectsToDelete != ObjectsToDelete.Num())
+		{
+			// If delete objects failed, then fallback to force delete
+			NumObjectsToDelete = ObjectTools::ForceDeleteObjects(ObjectsToDelete, bShowConfirmation);
+		}
+	}
+
+	const int32 NumPackagesToDelete = PackageFilesToDelete.Num();
+	if (NumPackagesToDelete > 0)
+	{
+		TArray<UPackage*> PackagePointers;
+		for (const TWeakObjectPtr<UPackage>& PkgIt : PackageFilesToDelete)
+		{
+			if (UPackage* Package = PkgIt.Get())
+			{
+				PackagePointers.Add(Package);
+			}
+		}
+
+		if (PackagePointers.Num() > 0)
+		{
+			constexpr bool bPerformReferenceCheck = true;
+			ObjectTools::CleanupAfterSuccessfulDelete(PackagePointers, bPerformReferenceCheck);
+		}
+	}
+
+	const int32 TotalDeletedObjects = NumPackagesToDelete + NumObjectsToDelete;
+	if (TotalDeletedObjects != AssetsToDelete.Num())
+	{
+		STORM_SYNC_IMPORT_LOG(Warning, TEXT("Failed to delete assets (Deleted %d assets while we were expecting to delete %d assets)"), TotalDeletedObjects, AssetsToDelete.Num());
+		return false;
+	}
+#endif
+
+	return true;
+}
+
+int32 UStormSyncImportSubsystem::MakePackagesWriteable(TConstArrayView<UPackage*> InPackages)
+{
+	int32 PackagesMadeWriteable = 0;
+
+	for (UPackage* Package : InPackages)
+	{
+		if (!Package)
+		{
+			continue;
+		}
+
+		FString Filename;
+		if (!FPackageName::DoesPackageExist(Package->GetName(), &Filename))
+		{
+			continue;
+		}
+
+		// Remove the read only flag from the current file attributes
+		if (FPlatformFileManager::Get().GetPlatformFile().SetReadOnly(*Filename, false))
+		{
+			++PackagesMadeWriteable;
+		}
+	}
+
+	return PackagesMadeWriteable;
+}
+
+bool UStormSyncImportSubsystem::IsAssetCurrentlyBeingEdited(const TSharedPtr<IToolkit>& InAssetEditor, const UObject* InAsset)
+{
+#if WITH_EDITOR
+	if (!InAsset)
+	{
+		return false;
+	}
+
+	if (InAssetEditor.IsValid() && InAssetEditor->IsAssetEditor())
+	{
+		if (const TArray<UObject*>* EditedObjects = InAssetEditor->GetObjectsCurrentlyBeingEdited())
+		{
+			for (const UObject* EditedObject : *EditedObjects)
+			{
+				if (EditedObject == InAsset)
+				{
+					return true;
+				}
+			}
+		}
+	}
+#endif
+
+	return false;
+}
+
+bool UStormSyncImportSubsystem::WriteFile(const FString& DestFilepath, const uint64& FileSize, uint8* FileBuffer)
+{
+	TUniquePtr<FArchive> AssetHandle(IFileManager::Get().CreateFileWriter(*DestFilepath));
+	if (!AssetHandle.IsValid())
+	{
+		return false;
+	}
+
+	STORM_SYNC_IMPORT_LOG(Verbose, TEXT("UStormSyncImportSubsystem::WriteFile - Creating file `%s`"), *DestFilepath);
+
+	// Write to asset
+	AssetHandle->Serialize(FileBuffer, FileSize);
+
+	// Close the file
+	AssetHandle->Close();
+	AssetHandle.Reset();
+
+	return true;
+}
+
+// Copied from FlushPackageLoading in ConcertSyncClientUtil.cpp.
+void UStormSyncImportSubsystem::FlushPackageLoading(const FString& InPackageName, bool bInForceBulkDataLoad)
+{	
+	if (UPackage* ExistingPackage = FindPackage(nullptr, *InPackageName))
+	{
+		if (!ExistingPackage->IsFullyLoaded())
+		{
+			FlushAsyncLoading();
+			ExistingPackage->FullyLoad();
+		}
+
+		if (bInForceBulkDataLoad)
+		{
+			ResetLoaders(ExistingPackage);
+		}
+		else if (ExistingPackage->GetLinker())
+		{
+			ExistingPackage->GetLinker()->Detach();
+		}
+	}
+}
+
+
+// Copied from HotReloadPackages in ConcertSyncClientUtil.cpp.
+void UStormSyncImportSubsystem::HotReloadPackages(const TArray<FStormSyncEditorFileReport>& InExtractedFileReports, bool bInInteractiveHotReload)
+{
+	TArray<FName> PackageNames;
+	PackageNames.Reserve(InExtractedFileReports.Num());
+	for (const FStormSyncEditorFileReport& FileReport : InExtractedFileReports)
+	{
+		if (FileReport.bSuccess)
+		{
+			PackageNames.Add(FileReport.PackageName);
+		}
+	}
+	
+	if (PackageNames.IsEmpty())
+	{
+		return;
+	}
+
+	STORM_SYNC_IMPORT_LOG(Display, TEXT("UStormSyncImportSubsystem::HotReloadPackages - Reloading %d packages."), PackageNames.Num());
+	
+	// Flush loading and clean-up any temporary placeholder packages (due to a package previously being missing on disk)
+	FlushAsyncLoading();
+	{
+		bool bRunGC = false;
+		for (const FName& PackageName : PackageNames)
+		{
+			bRunGC |= FLinkerLoad::RemoveKnownMissingPackage(PackageName);
+		}
+		if (bRunGC)
+		{
+			STORM_SYNC_IMPORT_LOG(Display, TEXT("UStormSyncImportSubsystem::HotReloadPackages - Garbage Collecting..."));
+			CollectGarbage(GARBAGE_COLLECTION_KEEPFLAGS);
+		}
+	}
+	
+	// Find the packages in-memory to content hot-reload
+	TArray<UPackage*> ExistingPackages;
+	ExistingPackages.Reserve(PackageNames.Num());
+
+	for (const FName& PackageName : PackageNames)
+	{
+		if (UPackage* ExistingPackage = FindPackage(nullptr, *PackageName.ToString()))
+		{
+			if (ExistingPackage->HasAnyPackageFlags(PKG_NewlyCreated))
+			{
+				ExistingPackage->ClearPackageFlags(PKG_NewlyCreated);
+			}
+			ExistingPackages.Add(ExistingPackage);
+		}
+	}
+
+	if (!ExistingPackages.IsEmpty())
+	{
+		FlushRenderingCommands();
+
+		FText ErrorMessage;
+#if WITH_EDITOR
+		const UPackageTools::EReloadPackagesInteractionMode InteractionMode = bInInteractiveHotReload ?
+			UPackageTools::EReloadPackagesInteractionMode::Interactive :
+			UPackageTools::EReloadPackagesInteractionMode::AssumePositive;
+
+		UPackageTools::ReloadPackages(ExistingPackages, ErrorMessage,InteractionMode);
+#endif
+
+		if (!ErrorMessage.IsEmpty())
+		{
+			STORM_SYNC_IMPORT_LOG(Error, TEXT("UStormSyncImportSubsystem::HotReloadPackages: %s"), *ErrorMessage.ToString());
+		}
+	}
+}
+
+void UStormSyncImportSubsystem::SplitImportFiles(TConstArrayView<FStormSyncImportFileInfo> InFiles, TArray<const FStormSyncImportFileInfo*>& OutExistingFiles, TArray<const FStormSyncImportFileInfo*>& OutNewFiles)
+{
+	if (InFiles.IsEmpty())
+	{
+		return;
+	}
+
+	// Worst case applied to both
+	OutExistingFiles.Empty(InFiles.Num());
+	OutNewFiles.Empty(InFiles.Num());
+
+	for (const FStormSyncImportFileInfo& FileInfo : InFiles)
+	{
+		if (FPaths::FileExists(FileInfo.DestFilepath))
+		{
+			OutExistingFiles.Add(&FileInfo);
+		}
+		else
+		{
+			OutNewFiles.Add(&FileInfo);
+		}
+	}
+}
+
+#undef LOCTEXT_NAMESPACE
diff --git a/Engine/Plugins/Experimental/StormSync/Source/StormSyncImport/Private/Subsystems/StormSyncImportWorldSubsystem.cpp b/Engine/Plugins/Experimental/StormSync/Source/StormSyncImport/Private/Subsystems/StormSyncImportWorldSubsystem.cpp
new file mode 100644
--- /dev/null
+++ b/Engine/Plugins/Experimental/StormSync/Source/StormSyncImport/Private/Subsystems/StormSyncImportWorldSubsystem.cpp
@@ -0,0 +1,68 @@
+// Copyright Epic Games, Inc. All Rights Reserved.
+
+
+#include "Subsystems/StormSyncImportWorldSubsystem.h"
+
+#include "Engine/World.h"
+#include "StormSyncCoreDelegates.h"
+#include "StormSyncImportLog.h"
+#include "Subsystems/StormSyncImportSubsystem.h"
+#include "Tasks/StormSyncImportBufferTask.h"
+#include "Tasks/StormSyncImportFilesTask.h"
+
+void UStormSyncImportWorldSubsystem::Initialize(FSubsystemCollectionBase& Collection)
+{
+	Super::Initialize(Collection);
+	STORM_SYNC_IMPORT_LOG(Verbose, TEXT("UStormSyncImportWorldSubsystem::Initialize (World: %s) - %s"), *GetNameSafe(GetWorld()), *GetName())
+
+	FStormSyncCoreDelegates::OnRequestImportFile.AddUObject(this, &UStormSyncImportWorldSubsystem::HandleImportFile);
+	FStormSyncCoreDelegates::OnRequestImportBuffer.AddUObject(this, &UStormSyncImportWorldSubsystem::HandleImportBuffer);
+}
+
+void UStormSyncImportWorldSubsystem::Deinitialize()
+{
+	Super::Deinitialize();
+	STORM_SYNC_IMPORT_LOG(Verbose, TEXT("UStormSyncImportWorldSubsystem::Deinitialize (World: %s) - %s"), *GetNameSafe(GetWorld()), *GetName())
+
+	FStormSyncCoreDelegates::OnRequestImportFile.RemoveAll(this);
+	FStormSyncCoreDelegates::OnRequestImportBuffer.RemoveAll(this);
+}
+
+void UStormSyncImportWorldSubsystem::HandleImportFile(const FString& InFilename) const
+{
+	STORM_SYNC_IMPORT_LOG(Verbose, TEXT("UStormSyncImportWorldSubsystem::HandleImportFile InFilename: %s (World: %s)"), *InFilename, *GetNameSafe(GetWorld()))
+	const UWorld* World = GetWorld();
+	if (!World)
+	{
+		STORM_SYNC_IMPORT_LOG(Error, TEXT("UStormSyncImportWorldSubsystem::ImportNextTick failed because of invalid world (World: %s)"), *GetNameSafe(GetWorld()))
+		return;
+	}
+
+	UStormSyncImportSubsystem::Get().EnqueueImportTask(MakeShared<FStormSyncImportFilesTask>(InFilename), GetWorld());
+}
+
+void UStormSyncImportWorldSubsystem::HandleImportBuffer(const FStormSyncPackageDescriptor& InPackageDescriptor, const FStormSyncBufferPtr& InBuffer) const
+{
+	if (!InBuffer.IsValid())
+	{
+		STORM_SYNC_IMPORT_LOG(Error, TEXT("UStormSyncImportSubsystem::HandleImportBuffer failed from invalid buffer"))
+		return;
+	}
+	
+	STORM_SYNC_IMPORT_LOG(
+		Verbose,
+		TEXT("UStormSyncImportWorldSubsystem::HandleImportBuffer InPackageDescriptor: %s, BufferSize: %d (World: %s)"),
+		*InPackageDescriptor.ToString(),
+		InBuffer->Num(),
+		*GetNameSafe(GetWorld())
+	)
+
+	const UWorld* World = GetWorld();
+	if (!World)
+	{
+		STORM_SYNC_IMPORT_LOG(Error, TEXT("UStormSyncImportWorldSubsystem::HandleImportBuffer failed because of invalid world (World: %s)"), *GetNameSafe(GetWorld()))
+		return;
+	}
+
+	UStormSyncImportSubsystem::Get().EnqueueImportTask(MakeShared<FStormSyncImportBufferTask>(InPackageDescriptor, InBuffer), GetWorld());
+}
diff --git a/Engine/Plugins/Experimental/StormSync/Source/StormSyncImport/Private/Tasks/StormSyncImportBufferTask.cpp b/Engine/Plugins/Experimental/StormSync/Source/StormSyncImport/Private/Tasks/StormSyncImportBufferTask.cpp
new file mode 100644
--- /dev/null
+++ b/Engine/Plugins/Experimental/StormSync/Source/StormSyncImport/Private/Tasks/StormSyncImportBufferTask.cpp
@@ -0,0 +1,18 @@
+// Copyright Epic Games, Inc. All Rights Reserved.
+
+#include "Tasks/StormSyncImportBufferTask.h"
+
+#include "StormSyncImportLog.h"
+#include "Subsystems/StormSyncImportSubsystem.h"
+
+void FStormSyncImportBufferTask::Run()
+{
+	if (!Buffer.IsValid())
+	{
+		STORM_SYNC_IMPORT_LOG(Error, TEXT("FStormSyncImportBufferTask::Run failed on invalid buffer"))
+		return;
+	}
+
+	STORM_SYNC_IMPORT_LOG(Display, TEXT("FStormSyncImportBufferTask::Run for buffer of size %d"), Buffer->Num())
+	UStormSyncImportSubsystem::Get().PerformBufferImport(PackageDescriptor, MoveTemp(Buffer));
+}
diff --git a/Engine/Plugins/Experimental/StormSync/Source/StormSyncImport/Private/Tasks/StormSyncImportBufferTask.h b/Engine/Plugins/Experimental/StormSync/Source/StormSyncImport/Private/Tasks/StormSyncImportBufferTask.h
new file mode 100644
--- /dev/null
+++ b/Engine/Plugins/Experimental/StormSync/Source/StormSyncImport/Private/Tasks/StormSyncImportBufferTask.h
@@ -0,0 +1,27 @@
+// Copyright Epic Games, Inc. All Rights Reserved.
+
+#pragma once
+
+#include "CoreMinimal.h"
+#include "StormSyncPackageDescriptor.h"
+#include "Tasks/IStormSyncImportTask.h"
+
+/** Import files from Buffer implementation for tasks that need delayed execution */
+class FStormSyncImportBufferTask final : public IStormSyncImportSubsystemTask
+{
+public:
+	explicit FStormSyncImportBufferTask(const FStormSyncPackageDescriptor& InPackageDescriptor, const TSharedPtr<TArray<uint8>, ESPMode::ThreadSafe>& InBuffer)
+		: PackageDescriptor(InPackageDescriptor)
+		, Buffer(InBuffer)
+	{
+	}
+
+	virtual void Run() override;
+
+private:
+	/** Metadata info about buffer being extracted */
+	FStormSyncPackageDescriptor PackageDescriptor;
+	
+	/** Holds a shared ptr to the serialized buffer */
+	TSharedPtr<TArray<uint8>, ESPMode::ThreadSafe> Buffer = nullptr;
+};
\ No newline at end of file
diff --git a/Engine/Plugins/Experimental/StormSync/Source/StormSyncImport/Public/IStormSyncImportModule.h b/Engine/Plugins/Experimental/StormSync/Source/StormSyncImport/Public/IStormSyncImportModule.h
new file mode 100644
--- /dev/null
+++ b/Engine/Plugins/Experimental/StormSync/Source/StormSyncImport/Public/IStormSyncImportModule.h
@@ -0,0 +1,34 @@
+// Copyright Epic Games, Inc. All Rights Reserved.
+
+#pragma once
+
+#include "Modules/ModuleInterface.h"
+#include "Modules/ModuleManager.h"
+
+/** Interface for StormSyncImport module entry point */
+class STORMSYNCIMPORT_API IStormSyncImportModule : public IModuleInterface
+{
+public:
+	/**
+	 * Singleton-like access to this module's interface.  This is just for convenience!
+	 * Beware of calling this during the shutdown phase, though.  Your module might have been unloaded already.
+	 *
+	 * @return Returns singleton instance, loading the module on demand if needed
+	 */
+	static inline IStormSyncImportModule& Get()
+	{
+		static const FName ModuleName = "StormSyncImport";
+		return FModuleManager::LoadModuleChecked<IStormSyncImportModule>(ModuleName);
+	}
+
+	/**
+	 * Checks to see if this module is loaded and ready.  It is only valid to call Get() during shutdown if IsAvailable() returns true.
+	 *
+	 * @return True if the module is loaded and ready to use
+	 */
+	static inline bool IsAvailable()
+	{
+		static const FName ModuleName = "StormSyncImport";
+		return FModuleManager::Get().IsModuleLoaded(ModuleName);
+	}
+};
diff --git a/Engine/Plugins/Experimental/StormSync/Source/StormSyncTests/Private/StormSyncTests.cpp b/Engine/Plugins/Experimental/StormSync/Source/StormSyncTests/Private/StormSyncTests.cpp
new file mode 100644
--- /dev/null
+++ b/Engine/Plugins/Experimental/StormSync/Source/StormSyncTests/Private/StormSyncTests.cpp
@@ -0,0 +1,16 @@
+// Copyright Epic Games, Inc. All Rights Reserved.
+
+#include "StormSyncTests.h"
+
+#include "Modules/ModuleManager.h"
+
+void FStormSyncTestsModule::StartupModule()
+{
+}
+
+void FStormSyncTestsModule::ShutdownModule()
+{
+
+}
+
+IMPLEMENT_MODULE(FStormSyncTestsModule, StormSyncTests)
diff --git a/Engine/Plugins/Experimental/StormSync/Source/StormSyncTests/Public/StormSyncTests.h b/Engine/Plugins/Experimental/StormSync/Source/StormSyncTests/Public/StormSyncTests.h
new file mode 100644
--- /dev/null
+++ b/Engine/Plugins/Experimental/StormSync/Source/StormSyncTests/Public/StormSyncTests.h
@@ -0,0 +1,16 @@
+// Copyright Epic Games, Inc. All Rights Reserved.
+
+#pragma once
+
+#include "Modules/ModuleInterface.h"
+
+/** Main entry point and implementation of StormSyncCore TestSuite UncookedOnly module. */
+class FStormSyncTestsModule : public IModuleInterface
+{
+public:
+
+	//~ Begin IModuleInterface interface
+	virtual void StartupModule() override;
+	virtual void ShutdownModule() override;
+	//~ End IModuleInterface interface
+};
diff --git a/Engine/Plugins/Experimental/StormSync/Source/StormSyncTransportClient/Private/StormSyncTransportClientModule.cpp b/Engine/Plugins/Experimental/StormSync/Source/StormSyncTransportClient/Private/StormSyncTransportClientModule.cpp
new file mode 100644
--- /dev/null
+++ b/Engine/Plugins/Experimental/StormSync/Source/StormSyncTransportClient/Private/StormSyncTransportClientModule.cpp
@@ -0,0 +1,301 @@
+// Copyright Epic Games, Inc. All Rights Reserved.
+
+#include "StormSyncTransportClientModule.h"
+
+#include "StormSyncCommandLineUtils.h"
+#include "StormSyncCoreUtils.h"
+#include "IStormSyncTransportCoreModule.h"
+#include "IStormSyncTransportLocalEndpoint.h"
+#include "MessageEndpoint.h"
+#include "Misc/CoreDelegates.h"
+#include "StormSyncTransportClientEndpoint.h"
+#include "StormSyncTransportClientLog.h"
+#include "StormSyncTransportMessages.h"
+
+#define LOCTEXT_NAMESPACE "StormSyncTransportClientModule"
+
+void FStormSyncTransportClientModule::StartupModule()
+{
+	ClientEndpoint = CreateClientLocalEndpoint(TEXT("Client"));
+
+	if (GIsEditor && !IsRunningCommandlet())
+	{
+		RegisterConsoleCommands();
+	}
+	
+	// Register for engine initialization completed so we can broadcast presence over the network from this client to other clients
+	FCoreDelegates::OnFEngineLoopInitComplete.AddRaw(this, &FStormSyncTransportClientModule::OnEngineLoopInitComplete);
+}
+
+void FStormSyncTransportClientModule::ShutdownModule()
+{
+	FCoreDelegates::OnFEngineLoopInitComplete.RemoveAll(this);
+	
+	if (ClientEndpoint.IsValid())
+	{
+		ClientEndpoint.Reset();
+	}
+
+	UnregisterConsoleCommands();
+}
+
+TSharedPtr<IStormSyncTransportClientLocalEndpoint> FStormSyncTransportClientModule::CreateClientLocalEndpoint(const FString& InEndpointFriendlyName) const
+{
+	TSharedPtr<FStormSyncTransportClientEndpoint, ESPMode::ThreadSafe> Endpoint = MakeShared<FStormSyncTransportClientEndpoint>();
+	check(Endpoint.IsValid());
+	Endpoint->InitializeMessaging(TEXT("Client"));
+	return Endpoint;
+}
+
+FString FStormSyncTransportClientModule::GetClientEndpointMessageAddressId() const
+{
+	return ClientEndpoint.IsValid() && ClientEndpoint->IsRunning() ? ClientEndpoint->GetMessageEndpoint()->GetAddress().ToString() : TEXT("");
+}
+
+FMessageEndpointSharedPtr FStormSyncTransportClientModule::GetClientMessageEndpoint() const
+{
+	if (!ClientEndpoint.IsValid())
+	{
+		STORM_SYNC_CLIENT_LOG(Warning, TEXT("FStormSyncTransportClientModule::GetClientMessageEndpoint - Client endpoint not valid"))
+		return nullptr;
+	}
+
+	if (!ClientEndpoint->IsRunning())
+	{
+		STORM_SYNC_CLIENT_LOG(Warning, TEXT("FStormSyncTransportClientModule::GetClientMessageEndpoint - Client endpoint not running"))
+		return nullptr;
+	}
+
+	return ClientEndpoint->GetMessageEndpoint();
+}
+
+void FStormSyncTransportClientModule::SynchronizePackages(const FStormSyncPackageDescriptor& InPackageDescriptor, const TArray<FName>& InPackageNames) const
+{
+	const FMessageEndpointSharedPtr MessageEndpoint = GetClientMessageEndpoint();
+
+	TArray<FName> LocalPackageNames = InPackageNames;
+	FStormSyncPackageDescriptor LocalPackageDescriptor = InPackageDescriptor;
+
+	FStormSyncCoreUtils::GetAvaFileDependenciesAsync(InPackageNames)
+	.Then([MessageEndpoint, PackageNames = MoveTemp(LocalPackageNames), PackageDescriptor = MoveTemp(LocalPackageDescriptor)](const TFuture<TArray<FStormSyncFileDependency>> Result)
+	{
+		const TArray<FStormSyncFileDependency> FileDependencies = Result.Get();
+		if (FileDependencies.IsEmpty())
+		{
+			STORM_SYNC_CLIENT_LOG(Error, TEXT("FStormSyncTransportClientModule::SynchronizePackages - Async FileDependencies is empty, something went wrong"))
+			return;
+		}
+
+		if (!MessageEndpoint.IsValid())
+		{
+			STORM_SYNC_CLIENT_LOG(Error, TEXT("FStormSyncTransportClientModule::SynchronizePackages - Unable to get client message endpoint"))
+			return;
+		}
+
+		// Build sync request message that is going to be sent over the network for a specific recipient
+		TUniquePtr<FStormSyncTransportSyncRequest> SyncRequestMessage(FMessageEndpoint::MakeMessage<FStormSyncTransportSyncRequest>(PackageNames, PackageDescriptor));
+		if (!SyncRequestMessage.IsValid())
+		{
+			STORM_SYNC_CLIENT_LOG(Error, TEXT("FStormSyncTransportClientModule::SynchronizePackages - Push request message is invalid"))
+			return;
+		}
+
+		SyncRequestMessage->PackageDescriptor.Dependencies = FileDependencies;
+
+		STORM_SYNC_CLIENT_LOG(Display, TEXT("FStormSyncTransportClientModule::SynchronizePackages - FileDependencies: %d"), FileDependencies.Num())
+		STORM_SYNC_CLIENT_LOG(Display, TEXT("FStormSyncTransportClientModule::SynchronizePackages - Message: %s"), *SyncRequestMessage->ToString())
+		STORM_SYNC_CLIENT_LOG(Display, TEXT("FStormSyncTransportClientModule::SynchronizePackages - Syncing package descriptor %s"), *SyncRequestMessage->PackageDescriptor.ToString())
+		MessageEndpoint->Publish(SyncRequestMessage.Release());
+	});
+}
+
+void FStormSyncTransportClientModule::PushPackages(const FStormSyncPackageDescriptor& InPackageDescriptor, const TArray<FName>& InPackageNames, const FMessageAddress& InMessageAddress, const FOnStormSyncPushComplete& InDoneDelegate) const
+{
+	STORM_SYNC_CLIENT_LOG(Display, TEXT("FStormSyncTransportClientModule::PushPackages - PackageDescriptor: %s, InPackageNames: %d, MessageAddressId: %s"), *InPackageDescriptor.ToString(), InPackageNames.Num(), *InMessageAddress.ToString())
+
+	const FMessageEndpointSharedPtr MessageEndpoint = GetClientMessageEndpoint();
+	if (!ClientEndpoint.IsValid())
+	{
+		STORM_SYNC_CLIENT_LOG(Error, TEXT("FStormSyncTransportClientModule::PushPackages - Unable to get client endpoint"))
+		return;
+	}
+
+	ClientEndpoint->RequestPushPackages(InMessageAddress, InPackageDescriptor, InPackageNames, InDoneDelegate);
+}
+
+void FStormSyncTransportClientModule::PullPackages(const FStormSyncPackageDescriptor& InPackageDescriptor, const TArray<FName>& InPackageNames, const FMessageAddress& InMessageAddress, const FOnStormSyncPullComplete& InDoneDelegate) const
+{
+	STORM_SYNC_CLIENT_LOG(Display, TEXT("FStormSyncTransportClientModule::PullPackages - PackageDescriptor: %s, InPackageNames: %d, MessageAddressId: %s"), *InPackageDescriptor.ToString(), InPackageNames.Num(), *InMessageAddress.ToString())
+
+	const FMessageEndpointSharedPtr MessageEndpoint = GetClientMessageEndpoint();
+
+	if (!ClientEndpoint.IsValid())
+	{
+		STORM_SYNC_CLIENT_LOG(Error, TEXT("FStormSyncTransportClientModule::PullPackages - Unable to get client endpoint"))
+		return;
+	}
+
+	ClientEndpoint->RequestPullPackages(InMessageAddress, InPackageDescriptor, InPackageNames, InDoneDelegate);
+}
+
+void FStormSyncTransportClientModule::RequestPackagesStatus(const FMessageAddress& InRemoteAddress, const TArray<FName>& InPackageNames, const FOnStormSyncRequestStatusComplete& InDoneDelegate) const
+{
+	STORM_SYNC_CLIENT_LOG(Display, TEXT("FStormSyncTransportClientModule::RequestPackagesStatus - InRemoteAddress: %s, InPackageNames: %d"), *InRemoteAddress.ToString(), InPackageNames.Num())
+
+	const FMessageEndpointSharedPtr MessageEndpoint = GetClientMessageEndpoint();
+	if (!ClientEndpoint.IsValid())
+	{
+		STORM_SYNC_CLIENT_LOG(Error, TEXT("FStormSyncTransportClientModule::RequestPackagesStatus - Unable to get client endpoint"))
+		return;
+	}
+
+	ClientEndpoint->RequestStatus(InRemoteAddress, InPackageNames, InDoneDelegate);
+}
+
+void FStormSyncTransportClientModule::RegisterConsoleCommands()
+{
+	ConsoleCommands.Add(IConsoleManager::Get().RegisterConsoleCommand(
+		TEXT("StormSync.Client.Ping"),
+		TEXT("Sends a ping message on Storm Sync message bus"),
+		FConsoleCommandWithArgsDelegate::CreateRaw(this, &FStormSyncTransportClientModule::ExecutePing),
+		ECVF_Default
+	));
+
+	ConsoleCommands.Add(IConsoleManager::Get().RegisterConsoleCommand(
+		TEXT("StormSync.Client.SyncPak"),
+		TEXT("Synchronize a pak file on Storm Sync message bus. Usage: <PackageName>... [-name=<AvaPackageName> -version=<AvaPackageVersion>] [-description=<AvaPackageVersion>] [-author=<AvaPackageAuthor>]"),
+		FConsoleCommandWithArgsDelegate::CreateRaw(this, &FStormSyncTransportClientModule::ExecuteSyncPak),
+		ECVF_Default
+	));
+
+	ConsoleCommands.Add(IConsoleManager::Get().RegisterConsoleCommand(
+		TEXT("StormSync.Client.Debug"),
+		TEXT("Prints out client address enpoint id to the console"),
+		// FConsoleCommandWithArgsDelegate::CreateRaw(this, &FStormSyncTransportClientModule::ExecuteSyncPak),
+		FConsoleCommandWithArgsDelegate::CreateLambda([this](const TArray<FString>& Args)
+		{
+			const FString AddressId = GetClientEndpointMessageAddressId();
+			STORM_SYNC_CLIENT_LOG(Display, TEXT("StormSync.Client.Debug - EndpointId: %s"), *AddressId)
+		}),
+		ECVF_Default
+	));
+
+	ConsoleCommands.Add(IConsoleManager::Get().RegisterConsoleCommand(
+		TEXT("StormSync.Client.Debug.Ping"),
+		TEXT("Sends status ping"),
+		// FConsoleCommandWithArgsDelegate::CreateRaw(this, &FStormSyncTransportClientModule::ExecuteSyncPak),
+		FConsoleCommandWithArgsDelegate::CreateLambda([this](const TArray<FString>& Args)
+		{
+			if (ClientEndpoint.IsValid())
+			{
+				const TSharedPtr<FMessageEndpoint, ESPMode::ThreadSafe> MessageEndpoint = ClientEndpoint->GetMessageEndpoint();
+				if (!MessageEndpoint.IsValid())
+				{
+					STORM_SYNC_CLIENT_LOG(Error, TEXT("StormSync.Client.Debug.Ping - Unable to send Connect Message cause Message Endpoint is invalid"))
+					return;
+				}
+
+				STORM_SYNC_CLIENT_LOG(Display, TEXT("StormSync.Client.Debug.Ping - Publish Client Connect Message ..."))
+				FStormSyncTransportStatusPing* Message = FMessageEndpoint::MakeMessage<FStormSyncTransportStatusPing>();
+				MessageEndpoint->Publish(Message);
+			}
+		}),
+		ECVF_Default
+	));
+	
+	ConsoleCommands.Add(IConsoleManager::Get().RegisterConsoleCommand(
+		TEXT("StormSync.Client.InstanceId"),
+		TEXT("Prints out local app instance id"),
+		FConsoleCommandWithArgsDelegate::CreateLambda([this](const TArray<FString>& Args)
+		{
+			const FString InstanceId = FApp::GetInstanceId().ToString();
+			STORM_SYNC_CLIENT_LOG(Display, TEXT("StormSync.Client.Debug - InstanceId: %s"), *InstanceId)
+		}),
+		ECVF_Default
+	));
+}
+
+void FStormSyncTransportClientModule::UnregisterConsoleCommands()
+{
+	for (IConsoleObject* Cmd : ConsoleCommands)
+	{
+		IConsoleManager::Get().UnregisterConsoleObject(Cmd);
+	}
+
+	ConsoleCommands.Empty();
+}
+
+void FStormSyncTransportClientModule::OnEngineLoopInitComplete()
+{
+	if (ClientEndpoint.IsValid())
+	{
+		const TSharedPtr<FMessageEndpoint, ESPMode::ThreadSafe> MessageEndpoint = ClientEndpoint->GetMessageEndpoint();
+		if (!MessageEndpoint.IsValid())
+		{
+			STORM_SYNC_CLIENT_LOG(Error, TEXT("FStormSyncTransportClientModule::OnEngineLoopInitComplete - Unable to send Connect Message cause Message Endpoint is invalid"))
+			return;
+		}
+
+		// We broadcast a message to notify others about this editor instance (this is required so that further "direct" send are received on the other end)
+		STORM_SYNC_CLIENT_LOG(Display, TEXT("FStormSyncTransportClientModule::OnEngineLoopInitComplete - Publish Client Connect Message ..."))
+		FStormSyncTransportStatusPing* Message = FMessageEndpoint::MakeMessage<FStormSyncTransportStatusPing>();
+		MessageEndpoint->Publish(Message);
+	}
+}
+
+void FStormSyncTransportClientModule::ExecutePing(const TArray<FString>& Args)
+{
+	const FString Argv = FString::Join(Args, TEXT(""));
+	STORM_SYNC_CLIENT_LOG(Display, TEXT("FStormSyncTransportClientModule::ExecutePing - %s"), *Argv)
+
+	const FMessageEndpointSharedPtr MessageEndpoint = GetClientMessageEndpoint();
+	if (!MessageEndpoint.IsValid())
+	{
+		STORM_SYNC_CLIENT_LOG(Warning, TEXT("FStormSyncTransportClientModule::ExecutePing - Unable to get client message endpoint"))
+		return;
+	}
+
+	MessageEndpoint->Publish(FMessageEndpoint::MakeMessage<FStormSyncTransportPingMessage>());
+}
+
+FStormSyncPackageDescriptor FStormSyncTransportClientModule::CreatePackageDescriptorFromCommandLine(const FString& Argv)
+{
+	FStormSyncPackageDescriptor PackageDescriptor;
+	if (!FParse::Value(*Argv, TEXT("-name="), PackageDescriptor.Name))
+	{
+		PackageDescriptor.Name = DefaultPakName;
+		STORM_SYNC_CLIENT_LOG(Display, TEXT("FStormSyncEditorModule::CreatePackageDescriptorFromCommandLine - Missing -name parameter, using default \"%s\""), *PackageDescriptor.Name);
+	}
+
+	FParse::Value(*Argv, TEXT("-version="), PackageDescriptor.Version);
+	FParse::Value(*Argv, TEXT("-description="), PackageDescriptor.Description);
+	FParse::Value(*Argv, TEXT("-author="), PackageDescriptor.Author);
+
+	return PackageDescriptor;
+}
+
+void FStormSyncTransportClientModule::ExecuteSyncPak(const TArray<FString>& Args)
+{
+	const FString Argv = FString::Join(Args, TEXT(" "));
+	STORM_SYNC_CLIENT_LOG(Display, TEXT("FStormSyncTransportClientModule::ExecuteSyncPak - Argv: %s"), *Argv)
+
+	// Parse command line.
+	TArray<FName> PackageNames;
+	FStormSyncCommandLineUtils::Parse(*Argv, PackageNames);
+
+	if (PackageNames.IsEmpty())
+	{
+		STORM_SYNC_CLIENT_LOG(Error, TEXT("FStormSyncTransportClientModule::ExecuteSyncPak - Missing at least one package name to sync."))
+		return;
+	}
+
+	// Create package descriptor pulling info from command line options
+	const FStormSyncPackageDescriptor PackageDescriptor = CreatePackageDescriptorFromCommandLine(Argv);
+
+	// Sync over network now
+	SynchronizePackages(PackageDescriptor, PackageNames);
+}
+
+IMPLEMENT_MODULE(FStormSyncTransportClientModule, StormSyncTransportClient)
+
+#undef LOCTEXT_NAMESPACE
diff --git a/Engine/Plugins/Experimental/StormSync/Source/StormSyncTransportClient/Public/StormSyncTransportClientLog.h b/Engine/Plugins/Experimental/StormSync/Source/StormSyncTransportClient/Public/StormSyncTransportClientLog.h
new file mode 100644
--- /dev/null
+++ b/Engine/Plugins/Experimental/StormSync/Source/StormSyncTransportClient/Public/StormSyncTransportClientLog.h
@@ -0,0 +1,12 @@
+// Copyright Epic Games, Inc. All Rights Reserved.
+
+#pragma once
+
+#include "CoreMinimal.h"
+
+STORMSYNCTRANSPORTCLIENT_API DECLARE_LOG_CATEGORY_EXTERN(LogStormSyncClient, Display, All);
+
+#define STORM_SYNC_CLIENT_LOG(Verbosity, Format, ...) \
+{ \
+    UE_LOG(LogStormSyncClient, Verbosity, Format, ##__VA_ARGS__); \
+}
diff --git a/Engine/Plugins/Experimental/StormSync/Source/StormSyncTransportCore/Private/StormSyncTransportCoreModule.cpp b/Engine/Plugins/Experimental/StormSync/Source/StormSyncTransportCore/Private/StormSyncTransportCoreModule.cpp
new file mode 100644
--- /dev/null
+++ b/Engine/Plugins/Experimental/StormSync/Source/StormSyncTransportCore/Private/StormSyncTransportCoreModule.cpp
@@ -0,0 +1,18 @@
+// Copyright Epic Games, Inc. All Rights Reserved.
+
+#include "StormSyncTransportCoreModule.h"
+
+#define LOCTEXT_NAMESPACE "StormSyncTransportCoreModule"
+
+
+void FStormSyncTransportCoreModule::StartupModule()
+{
+}
+
+void FStormSyncTransportCoreModule::ShutdownModule()
+{
+}
+
+IMPLEMENT_MODULE(FStormSyncTransportCoreModule, StormSyncTransportCore)
+
+#undef LOCTEXT_NAMESPACE
diff --git a/Engine/Plugins/Experimental/StormSync/Source/StormSyncTransportCore/Private/StormSyncTransportCoreModule.h b/Engine/Plugins/Experimental/StormSync/Source/StormSyncTransportCore/Private/StormSyncTransportCoreModule.h
new file mode 100644
--- /dev/null
+++ b/Engine/Plugins/Experimental/StormSync/Source/StormSyncTransportCore/Private/StormSyncTransportCoreModule.h
@@ -0,0 +1,12 @@
+// Copyright Epic Games, Inc. All Rights Reserved.
+
+#pragma once
+
+#include "IStormSyncTransportCoreModule.h"
+
+class FStormSyncTransportCoreModule : public IStormSyncTransportCoreModule
+{
+public:
+	virtual void StartupModule() override;
+	virtual void ShutdownModule() override;
+};
diff --git a/Engine/Plugins/Experimental/StormSync/Source/StormSyncTransportCore/Public/StormSyncTransportCoreLog.h b/Engine/Plugins/Experimental/StormSync/Source/StormSyncTransportCore/Public/StormSyncTransportCoreLog.h
new file mode 100644
--- /dev/null
+++ b/Engine/Plugins/Experimental/StormSync/Source/StormSyncTransportCore/Public/StormSyncTransportCoreLog.h
@@ -0,0 +1,12 @@
+// Copyright Epic Games, Inc. All Rights Reserved.
+
+#pragma once
+
+#include "CoreMinimal.h"
+
+STORMSYNCTRANSPORTCORE_API DECLARE_LOG_CATEGORY_EXTERN(LogStormSyncTransportCore, Display, All);
+
+#define STORM_SYNC_CORE_LOG(Verbosity, Format, ...) \
+{ \
+    UE_LOG(LogStormSyncTransportCore, Verbosity, Format, ##__VA_ARGS__); \
+}
diff --git a/Engine/Plugins/Experimental/StormSync/Source/StormSyncTransportServer/Private/ServiceDiscovery/StormSyncDiscoveryManager.cpp b/Engine/Plugins/Experimental/StormSync/Source/StormSyncTransportServer/Private/ServiceDiscovery/StormSyncDiscoveryManager.cpp
new file mode 100644
--- /dev/null
+++ b/Engine/Plugins/Experimental/StormSync/Source/StormSyncTransportServer/Private/ServiceDiscovery/StormSyncDiscoveryManager.cpp
@@ -0,0 +1,321 @@
+// Copyright Epic Games, Inc. All Rights Reserved.
+
+#include "StormSyncDiscoveryManager.h"
+
+#include "Async/TaskGraphInterfaces.h"
+#include "HAL/RunnableThread.h"
+#include "IStormSyncTransportClientModule.h"
+#include "IStormSyncTransportServerModule.h"
+#include "MessageEndpointBuilder.h"
+#include "StormSyncCoreDelegates.h"
+#include "StormSyncHeartbeatEmitter.h"
+#include "StormSyncTransportMessages.h"
+#include "StormSyncTransportServerLog.h"
+
+FStormSyncDiscoveryManager::FStormSyncDiscoveryManager(double InHeartbeatTimeout, double InInactiveSourceTimeout, float InTickInterval, bool bInEnableDiscoveryPeriodicPublish)
+	: bRunning(true)
+	, Thread(nullptr)
+	, DefaultHeartbeatTimeout(InHeartbeatTimeout)
+	, DefaultDeadSourceTimeout(InInactiveSourceTimeout)
+	, DefaultTickInterval(InTickInterval)
+	, bEnableDiscoveryPeriodicPublish(bInEnableDiscoveryPeriodicPublish)
+{
+	MessageEndpoint = FMessageEndpoint::Builder(TEXT("StormSyncMessageHeartbeatManager (StormSyncDiscoveryManager)"))
+		.ReceivingOnThread(ENamedThreads::GameThread)
+		.Handling<FStormSyncTransportConnectMessage>(this, &FStormSyncDiscoveryManager::HandleConnectMessage)
+		.Handling<FStormSyncTransportHeartbeatMessage>(this, &FStormSyncDiscoveryManager::HandleHeartbeatMessage);
+
+	bRunning = MessageEndpoint.IsValid();
+	if (bRunning)
+	{
+		STORM_SYNC_SERVER_LOG(Display, TEXT("FStormSyncDiscoveryManager::FStormSyncDiscoveryManager - Subscribe to messages"))
+		MessageEndpoint->Subscribe<FStormSyncTransportConnectMessage>();
+
+		STORM_SYNC_SERVER_LOG(Display, TEXT("FStormSyncDiscoveryManager::FStormSyncDiscoveryManager - Start Thread"))
+		Thread = FRunnableThread::Create(this, TEXT("StormSyncDiscoveryManager"));
+	}
+}
+
+FStormSyncDiscoveryManager::~FStormSyncDiscoveryManager()
+{
+	{
+		FScopeLock Lock(&ConnectionLastActiveSection);
+
+		// Disable the Endpoint message handling since the message could keep it alive a bit.
+		if (MessageEndpoint)
+		{
+			MessageEndpoint->Disable();
+			MessageEndpoint.Reset();
+		}
+	}
+
+	bRunning = false;
+
+	if (Thread)
+	{
+		Thread->Kill(true);
+		Thread = nullptr;
+	}
+}
+
+uint32 FStormSyncDiscoveryManager::Run()
+{
+	while (bRunning)
+	{
+		const double CurrentTime = FApp::GetCurrentTime();
+		
+		TArray<FStormSyncDelegateItem> DelegateQueue;
+		{
+			FScopeLock Lock(&ConnectionLastActiveSection);
+
+
+			TArray<FMessageAddress> DisconnectedAddresses;
+			for (TPair<FMessageAddress, FStormSyncConnectedMessageBusAddress>& Entry : ConnectedAddresses)
+			{
+				FMessageAddress MessageAddress = Entry.Key;
+				FStormSyncConnectedMessageBusAddress& ConnectedAddress = Entry.Value;
+				const FString MessageAddressUID = MessageAddress.ToString();
+
+				const bool bPrevState = ConnectedAddress.bIsValid;
+
+				const double ElapsedTime = CurrentTime - ConnectedAddress.LastActivityTime;
+				const bool bIsValid = ElapsedTime < DefaultHeartbeatTimeout;
+
+				if (bPrevState != bIsValid)
+				{
+					// Queue up a delegate event to be fired off later on
+					DelegateQueue.Add(FStormSyncDelegateItem(
+						FStormSyncDelegateItem::StateChange,
+						MessageAddress,
+						bIsValid ? EStormSyncConnectedDeviceState::State_Active : EStormSyncConnectedDeviceState::State_Unresponsive
+					));
+				}
+
+				// Update state in stored connection
+				ConnectedAddress.bIsValid = bIsValid;
+
+				// Connection starting to be unresponsive, can be a real disconnect or an occasional lag spike
+				if (!bIsValid)
+				{
+					// If we exceeded inactive time, consider this remote as disconnected
+					if (ElapsedTime > DefaultDeadSourceTimeout)
+					{
+						DisconnectedAddresses.Add(MessageAddress);
+					}
+				}
+			}
+
+			// Handle inactive addresses
+			for (FMessageAddress MessageAddress : DisconnectedAddresses)
+			{
+				STORM_SYNC_SERVER_LOG(Display, TEXT("FStormSyncDiscoveryManager::Run - %s became invalid"), *MessageAddress.ToString())
+
+				// Stop Heartbeat for this recipient
+				FStormSyncHeartbeatEmitter& HeartbeatEmitter = IStormSyncTransportServerModule::Get().GetHeartbeatEmitter();
+				HeartbeatEmitter.StopHeartbeat(MessageAddress, MessageEndpoint);
+
+				// Queue up a delegate event to notify about disconnection
+				FStormSyncDelegateItem Item = FStormSyncDelegateItem(FStormSyncDelegateItem::Disconnection, MessageAddress, EStormSyncConnectedDeviceState::State_Disconnected);
+				DelegateQueue.Add(Item);
+
+				// Actually remove the address now
+				ConnectedAddresses.Remove(MessageAddress);
+			}
+		}
+
+		// Handle any delegate to fire off now
+		BroadcastCoreDelegatesFromQueue(DelegateQueue);
+
+		// Handle periodic connect message publish if enabled
+		if (bEnableDiscoveryPeriodicPublish)
+		{
+			const double ElapsedTime = CurrentTime - LastPublishTime;
+			if (ElapsedTime > DefaultDeadSourceTimeout)
+			{
+				PublishConnectMessage();
+			}
+		}
+
+		FPlatformProcess::Sleep(FMath::Max<float>(0.1f, DefaultTickInterval));
+	}
+	return 0;
+}
+
+void FStormSyncDiscoveryManager::Stop()
+{
+	bRunning = false;
+}
+
+void FStormSyncDiscoveryManager::PublishConnectMessage()
+{
+	if (!MessageEndpoint.IsValid())
+	{
+		STORM_SYNC_SERVER_LOG(Error, TEXT("FStormSyncDiscoveryManager::PublishConnectMessage - Unable to send Connect Message cause Message Endpoint is invalid"))
+		return;
+	}
+
+	STORM_SYNC_SERVER_LOG(VeryVerbose, TEXT("FStormSyncTransportClientModule::PublishConnectMessage - Publish Connect Message ..."))
+	TUniquePtr<FStormSyncTransportConnectMessage> Message(FMessageEndpoint::MakeMessage<FStormSyncTransportConnectMessage>());
+	if (Message.IsValid())
+	{
+		Message->StormSyncServerAddressId = IStormSyncTransportServerModule::Get().GetServerEndpointMessageAddressId();
+		Message->StormSyncClientAddressId = IStormSyncTransportClientModule::Get().GetClientEndpointMessageAddressId();
+
+		LastPublishTime = FApp::GetCurrentTime();
+		MessageEndpoint->Publish(Message.Release());
+	}
+}
+
+TSharedPtr<FMessageEndpoint, ESPMode::ThreadSafe> FStormSyncDiscoveryManager::GetMessageEndpoint() const
+{
+	return MessageEndpoint;
+}
+
+void FStormSyncDiscoveryManager::BroadcastCoreDelegatesFromQueue(const TArray<FStormSyncDelegateItem>& InDelegateQueue)
+{
+	for (const FStormSyncDelegateItem& DelegateItem : InDelegateQueue)
+	{
+		if (DelegateItem.DelegateType == FStormSyncDelegateItem::EDelegateType::StateChange)
+		{
+			// Notify editor this remote state changed (either responsive or unresponsive)
+			FStormSyncCoreDelegates::OnServiceDiscoveryStateChange.Broadcast(DelegateItem.Address.ToString(), DelegateItem.State);
+		}
+		else if (DelegateItem.DelegateType == FStormSyncDelegateItem::EDelegateType::Disconnection)
+		{
+			// Notify Editor this remote is considered disconnected
+			FStormSyncCoreDelegates::OnServiceDiscoveryDisconnection.Broadcast(DelegateItem.Address.ToString());
+		}
+	}
+}
+
+void FStormSyncDiscoveryManager::HandleConnectMessage(const FStormSyncTransportConnectMessage& InMessage, const TSharedRef<IMessageContext, ESPMode::ThreadSafe>& MessageContext)
+{
+	const FMessageAddress SenderMessageAddress = MessageContext->GetSender();
+	STORM_SYNC_SERVER_LOG(VeryVerbose, TEXT("FStormSyncDiscoveryManager::HandleConnectMessage - Received connect Message from %s: %s"), *SenderMessageAddress.ToString(), *InMessage.ToString())
+
+	const FMessageAddress MessageEndpointAddress = MessageEndpoint->GetAddress();
+
+	STORM_SYNC_SERVER_LOG(VeryVerbose, TEXT("FStormSyncDiscoveryManager::HandleConnectMessage - SenderMessageAddress: %s, MessageEndpointAddress: %s (Same: %s)"), *SenderMessageAddress.ToString(), *MessageEndpointAddress.ToString(), SenderMessageAddress == MessageEndpointAddress ? TEXT("true") : TEXT("false"))
+
+	if (SenderMessageAddress != MessageEndpointAddress && !ConnectedAddresses.Contains(SenderMessageAddress))
+	{
+		// Add this message address to our list of connected message bus addresses
+		RegisterConnection(InMessage, SenderMessageAddress);
+	}
+}
+
+void FStormSyncDiscoveryManager::RegisterConnection(const FStormSyncTransportConnectMessage& InMessage, const FMessageAddress& InMessageAddress)
+{
+	{
+		FScopeLock ConnectionTimeLock(&ConnectionLastActiveSection);
+		check(ConnectedAddresses.Contains(InMessageAddress) == false);
+
+		FStormSyncConnectedMessageBusAddress ConnectedAddress;
+		ConnectedAddress.bIsValid = true;
+		ConnectedAddress.bIsServerRunning = !InMessage.StormSyncServerAddressId.IsEmpty();
+		ConnectedAddress.LastActivityTime = FPlatformTime::Seconds();
+
+		ConnectedAddresses.Add(InMessageAddress, ConnectedAddress);
+	}
+
+	// Start Heartbeat for this recipient
+	FStormSyncHeartbeatEmitter& HeartbeatEmitter = IStormSyncTransportServerModule::Get().GetHeartbeatEmitter();
+	HeartbeatEmitter.StartHeartbeat(InMessageAddress, MessageEndpoint);
+
+	// Send back connect message so that this recipient knows about this editor instance
+	STORM_SYNC_SERVER_LOG(Display, TEXT("FStormSyncDiscoveryManager::HandleConnectMessage - Send Connect Message to %s..."), *InMessageAddress.ToString())
+	TUniquePtr<FStormSyncTransportConnectMessage> Message(FMessageEndpoint::MakeMessage<FStormSyncTransportConnectMessage>());
+	if (Message.IsValid())
+	{
+		Message->StormSyncServerAddressId = IStormSyncTransportServerModule::Get().GetServerEndpointMessageAddressId();
+		Message->StormSyncClientAddressId = IStormSyncTransportClientModule::Get().GetClientEndpointMessageAddressId();
+		MessageEndpoint->Send(Message.Release(), InMessageAddress);
+	}
+
+	// Notify editor of incoming connection
+	FStormSyncConnectedDevice ConnectedDevice;
+	ConnectedDevice.State = EStormSyncConnectedDeviceState::State_Active;
+	ConnectedDevice.MessageAddressId = InMessageAddress.ToString();
+	ConnectedDevice.bIsServerRunning = !InMessage.StormSyncServerAddressId.IsEmpty();
+	ConnectedDevice.StormSyncServerAddressId = InMessage.StormSyncServerAddressId;
+	ConnectedDevice.StormSyncClientAddressId = InMessage.StormSyncClientAddressId;
+	ConnectedDevice.HostName = InMessage.HostName;
+	ConnectedDevice.ProjectName = InMessage.ProjectName;
+	ConnectedDevice.ProjectDir = InMessage.ProjectDir;
+	ConnectedDevice.InstanceType = InMessage.InstanceType;
+
+	FStormSyncCoreDelegates::OnServiceDiscoveryConnection.Broadcast(InMessageAddress.ToString(), ConnectedDevice);
+}
+
+void FStormSyncDiscoveryManager::HandleHeartbeatMessage(const FStormSyncTransportHeartbeatMessage& InMessage, const TSharedRef<IMessageContext, ESPMode::ThreadSafe>& MessageContext)
+{
+	const FMessageAddress MessageAddress = MessageContext->GetSender();
+
+	STORM_SYNC_SERVER_LOG(VeryVerbose, TEXT("FStormSyncDiscoveryManager::HandleHeartbeatMessage - Received Heartbeat Message from %s"), *MessageAddress.ToString())
+	
+	{
+		// Handle revive connection in case we receive heartbeats again, most likely meaning
+		// connection was marked here as inactive due to inactive timeout and cleaned up.
+		//
+		// This is likely caused by a debug session with breakpoints that are longer than the configured inactive timeout
+
+		bool bContainsAddress;
+		{
+			FScopeLock ConnectionTimeLock(&ConnectionLastActiveSection);
+			bContainsAddress = ConnectedAddresses.Contains(MessageAddress);
+		}
+
+		if (!bContainsAddress)
+		{
+			TUniquePtr<FStormSyncTransportConnectMessage> ReviveMessage(FMessageEndpoint::MakeMessage<FStormSyncTransportConnectMessage>());
+			if (ReviveMessage.IsValid())
+			{
+				ReviveMessage->StormSyncServerAddressId = IStormSyncTransportServerModule::Get().GetServerEndpointMessageAddressId();
+				ReviveMessage->StormSyncClientAddressId = IStormSyncTransportClientModule::Get().GetClientEndpointMessageAddressId();
+				RegisterConnection(*ReviveMessage.Release(), MessageAddress);
+			}
+		}
+	}
+
+	UpdateConnectionLastActive(MessageAddress);
+
+	// Update Server status and check if it changed for this recipient since last heartbeat, if so notify editor
+	if (UpdateServerStatus(MessageAddress, InMessage.bIsServerRunning))
+	{
+		// Notify editor this remote state changed (either running or stopped or unresponsive)
+		FStormSyncCoreDelegates::OnServiceDiscoveryServerStatusChange.Broadcast(MessageAddress.ToString(), InMessage.bIsServerRunning);
+	}
+}
+
+void FStormSyncDiscoveryManager::UpdateConnectionLastActive(const FMessageAddress& InAddress)
+{
+	FScopeLock ConnectionTimeLock(&ConnectionLastActiveSection);
+
+	if (ConnectedAddresses.Contains(InAddress))
+	{
+		FStormSyncConnectedMessageBusAddress& ConnectedMessageBusAddress = ConnectedAddresses.FindChecked(InAddress);
+		ConnectedMessageBusAddress.LastActivityTime = FPlatformTime::Seconds();
+	}
+}
+
+bool FStormSyncDiscoveryManager::UpdateServerStatus(const FMessageAddress& InAddress, const bool bIsServerRunning)
+{
+	FScopeLock ConnectionTimeLock(&ConnectionLastActiveSection);
+
+	bool bServerStatusChanged = false;
+
+	if (ConnectedAddresses.Contains(InAddress))
+	{
+		FStormSyncConnectedMessageBusAddress& ConnectedMessageBusAddress = ConnectedAddresses.FindChecked(InAddress);
+
+		// Server status changed for this recipient
+		if (ConnectedMessageBusAddress.bIsServerRunning != bIsServerRunning)
+		{
+			bServerStatusChanged = true;
+		}
+
+		ConnectedMessageBusAddress.bIsServerRunning = bIsServerRunning;
+	}
+
+	return bServerStatusChanged;
+}
diff --git a/Engine/Plugins/Experimental/StormSync/Source/StormSyncTransportServer/Public/StormSyncTransportServerLog.h b/Engine/Plugins/Experimental/StormSync/Source/StormSyncTransportServer/Public/StormSyncTransportServerLog.h
new file mode 100644
--- /dev/null
+++ b/Engine/Plugins/Experimental/StormSync/Source/StormSyncTransportServer/Public/StormSyncTransportServerLog.h
@@ -0,0 +1,12 @@
+// Copyright Epic Games, Inc. All Rights Reserved.
+
+#pragma once
+
+#include "CoreMinimal.h"
+
+STORMSYNCTRANSPORTSERVER_API DECLARE_LOG_CATEGORY_EXTERN(LogStormSyncServer, Display, All);
+
+#define STORM_SYNC_SERVER_LOG(Verbosity, Format, ...) \
+{ \
+    UE_LOG(LogStormSyncServer, Verbosity, Format, ##__VA_ARGS__); \
+}
