Storm Sync:
- Changed Warning to Display when unable to GetServerEndpoint (that is called in Module Startup, and was logging this warning if settings were not set up)

diff --git a/Engine/Plugins/Experimental/StormSync/Source/StormSyncTransportCore/Private/StormSyncTransportSettings.cpp b/Engine/Plugins/Experimental/StormSync/Source/StormSyncTransportCore/Private/StormSyncTransportSettings.cpp
--- a/Engine/Plugins/Experimental/StormSync/Source/StormSyncTransportCore/Private/StormSyncTransportSettings.cpp
+++ b/Engine/Plugins/Experimental/StormSync/Source/StormSyncTransportCore/Private/StormSyncTransportSettings.cpp
@@ -64,7 +64,7 @@ FString UStormSyncTransportSettings::GetServerEndpoint() const
 		if (!GetUdpMessagingUnicastEndpoint(UnicastHostname, UnicastPort))
 		{
 			// Either unable to get the setting or malformed
-			STORM_SYNC_CORE_LOG(Warning, TEXT("UStormSyncTransportSettings::GetServerEndpoint - Was unable to get UnicastEndpoint from UDPMessagingSettings"));
+			STORM_SYNC_CORE_LOG(Display, TEXT("UStormSyncTransportSettings::GetServerEndpoint - Was unable to get UnicastEndpoint from UDPMessagingSettings"));
 			return ServerEndpoint;
 		}
 
