Storm Sync:
First pass cleanup
- Move unexported to private
- Unexported classes that didn't need to be exported or could have an interface class instead

#changelist validated
#virtualized

diff --git a/Engine/Plugins/Experimental/StormSync/Source/StormSyncCore/Private/StormSyncCoreLog.cpp b/Engine/Plugins/Experimental/StormSync/Source/StormSyncCore/Private/StormSyncCoreLog.cpp
deleted file mode 100644
diff --git a/Engine/Plugins/Experimental/StormSync/Source/StormSyncCore/Private/StormSyncCoreLog.h b/Engine/Plugins/Experimental/StormSync/Source/StormSyncCore/Private/StormSyncCoreLog.h
new file mode 100644
--- /dev/null
+++ b/Engine/Plugins/Experimental/StormSync/Source/StormSyncCore/Private/StormSyncCoreLog.h
@@ -0,0 +1,12 @@
+// Copyright Epic Games, Inc. All Rights Reserved.
+
+#pragma once
+
+#include "Logging/LogMacros.h"
+
+DECLARE_LOG_CATEGORY_EXTERN(LogStormSyncCore, Display, All);
+
+#define STORM_SYNC_LOG(Verbosity, Format, ...) \
+{ \
+    UE_LOG(LogStormSyncCore, Verbosity, Format, ##__VA_ARGS__); \
+}
diff --git a/Engine/Plugins/Experimental/StormSync/Source/StormSyncCore/Private/StormSyncCoreModule.cpp b/Engine/Plugins/Experimental/StormSync/Source/StormSyncCore/Private/StormSyncCoreModule.cpp
--- a/Engine/Plugins/Experimental/StormSync/Source/StormSyncCore/Private/StormSyncCoreModule.cpp
+++ b/Engine/Plugins/Experimental/StormSync/Source/StormSyncCore/Private/StormSyncCoreModule.cpp
@@ -1,13 +1,8 @@
 // Copyright Epic Games, Inc. All Rights Reserved.
 
-#include "StormSyncCoreModule.h"
+#include "Modules/ModuleManager.h"
+#include "StormSyncCoreLog.h"
 
-void FStormSyncCoreModule::StartupModule()
-{
-}
+DEFINE_LOG_CATEGORY(LogStormSyncCore);
 
-void FStormSyncCoreModule::ShutdownModule()
-{
-}
-
-IMPLEMENT_MODULE(FStormSyncCoreModule, StormSyncCore)
+IMPLEMENT_MODULE(FDefaultModuleImpl, StormSyncCore)
diff --git a/Engine/Plugins/Experimental/StormSync/Source/StormSyncCore/Private/StormSyncCoreSettings.cpp b/Engine/Plugins/Experimental/StormSync/Source/StormSyncCore/Private/StormSyncCoreSettings.cpp
--- a/Engine/Plugins/Experimental/StormSync/Source/StormSyncCore/Private/StormSyncCoreSettings.cpp
+++ b/Engine/Plugins/Experimental/StormSync/Source/StormSyncCore/Private/StormSyncCoreSettings.cpp
@@ -4,19 +4,9 @@
 
 UStormSyncCoreSettings::UStormSyncCoreSettings()
 {
+	CategoryName = TEXT("Storm Sync");
+
 	IgnoredPackagesInternal.Add(TEXT("/Engine"));
 	IgnoredPackagesInternal.Add(TEXT("/Script"));
 	ExportDefaultNameFormatString = TEXT("%Y_%m_%d_%H%M%S");
 }
-
-FName UStormSyncCoreSettings::GetCategoryName() const
-{
-	return PluginCategoryName;
-}
-
-#if WITH_EDITOR
-FText UStormSyncCoreSettings::GetSectionText() const
-{
-	return NSLOCTEXT("StormSyncCorePlugin", "StormSyncCoreSettingsSection", "Core Settings");
-}
-#endif
diff --git a/Engine/Plugins/Experimental/StormSync/Source/StormSyncCore/Public/StormSyncCommandLineUtils.h b/Engine/Plugins/Experimental/StormSync/Source/StormSyncCore/Public/StormSyncCommandLineUtils.h
--- a/Engine/Plugins/Experimental/StormSync/Source/StormSyncCore/Public/StormSyncCommandLineUtils.h
+++ b/Engine/Plugins/Experimental/StormSync/Source/StormSyncCore/Public/StormSyncCommandLineUtils.h
@@ -2,7 +2,10 @@
 
 #pragma once
 
-#include "CoreMinimal.h"
+#include "Containers/ContainersFwd.h"
+#include "HAL/Platform.h"
+
+class FName;
 
 /** Provides command line helpers to help in dealing with command parameters */
 class STORMSYNCCORE_API FStormSyncCommandLineUtils
diff --git a/Engine/Plugins/Experimental/StormSync/Source/StormSyncCore/Public/StormSyncCommonTypes.h b/Engine/Plugins/Experimental/StormSync/Source/StormSyncCore/Public/StormSyncCommonTypes.h
--- a/Engine/Plugins/Experimental/StormSync/Source/StormSyncCore/Public/StormSyncCommonTypes.h
+++ b/Engine/Plugins/Experimental/StormSync/Source/StormSyncCore/Public/StormSyncCommonTypes.h
@@ -2,7 +2,7 @@
 
 #pragma once
 
-#include "CoreMinimal.h"
+#include "Templates/SharedPointer.h"
 #include "StormSyncCommonTypes.generated.h"
 
 // Common alias to a ThreadSafe SharedPtr to hold a Buffer (array of uint8)
diff --git a/Engine/Plugins/Experimental/StormSync/Source/StormSyncCore/Public/StormSyncCoreDelegates.h b/Engine/Plugins/Experimental/StormSync/Source/StormSyncCore/Public/StormSyncCoreDelegates.h
--- a/Engine/Plugins/Experimental/StormSync/Source/StormSyncCore/Public/StormSyncCoreDelegates.h
+++ b/Engine/Plugins/Experimental/StormSync/Source/StormSyncCore/Public/StormSyncCoreDelegates.h
@@ -2,7 +2,6 @@
 
 #pragma once
 
-#include "CoreMinimal.h"
 #include "StormSyncCommonTypes.h"
 #include "StormSyncPackageDescriptor.h"
 #include "StormSyncCoreDelegates.generated.h"
@@ -22,7 +21,7 @@ enum class EStormSyncConnectedDeviceState : uint8
 };
 
 /** UI related data holder for storm sync connected devices */
-struct STORMSYNCCORE_API FStormSyncConnectedDevice
+struct FStormSyncConnectedDevice
 {
 	/** Holds the state of the connection*/
 	EStormSyncConnectedDeviceState State = EStormSyncConnectedDeviceState::State_Disconnected;
@@ -62,7 +61,7 @@ struct STORMSYNCCORE_API FStormSyncConnectedDevice
 	FStormSyncConnectedDevice() = default;
 };
 
-/** 
+/**
  * FStormSyncCoreDelegates
  * 
  * Delegates used by the editor.
@@ -76,7 +75,7 @@ struct STORMSYNCCORE_API FStormSyncCoreDelegates
 	 * @param FileBuffer The raw file buffer (as a  shared ptr array of bytes) for the incoming pak buffer
 	 */
 	DECLARE_MULTICAST_DELEGATE_TwoParams(FOnRequestImportBuffer, const FStormSyncPackageDescriptor&, const FStormSyncBufferPtr&);
-	
+
 	/**
 	 * Delegate type for when one an incoming pak extraction process starts
 	 *
@@ -90,7 +89,7 @@ struct STORMSYNCCORE_API FStormSyncCoreDelegates
 	 * @param InFilename The absolute filename path for the imported file
 	 */
 	DECLARE_MULTICAST_DELEGATE_OneParam(FOnFileImported, const FString& /*InFilename*/);
-	
+
 	/**
 	 * Delegate type for when one asset is extracted from an incoming pak
 	 *
@@ -107,7 +106,7 @@ struct STORMSYNCCORE_API FStormSyncCoreDelegates
 	 * @param PackageNums The total number of files (or package names) to include in the buffer
 	 */
 	DECLARE_MULTICAST_DELEGATE_ThreeParams(FOnPreStartSendingBuffer, const FString&, const FString&, int32);
-	
+
 	/**
 	 * Delegate type for a client socket opening a tcp stream and sending a buffer to a remote server
 	 *
@@ -115,7 +114,7 @@ struct STORMSYNCCORE_API FStormSyncCoreDelegates
 	 * @param BufferSize The total number of bytes the server should expect, representing the buffer size
 	 */
 	DECLARE_MULTICAST_DELEGATE_TwoParams(FOnStartSendingBuffer, const FString&, int32);
-	
+
 	/**
 	 * Delegate type for a client socket receiving response from server, indicating how many bytes it received so far
 	 *
@@ -123,7 +122,7 @@ struct STORMSYNCCORE_API FStormSyncCoreDelegates
 	 * @param BytesCount The number of bytes received so far by the remote server
 	 */
 	DECLARE_MULTICAST_DELEGATE_TwoParams(FOnReceivingBytes, const FString&, int32);
-	
+
 	/**
 	 * Delegate type for a client socket receiving "transfer complete" response from server, indicating the tcp transfer is done.
 	 *
@@ -140,7 +139,7 @@ struct STORMSYNCCORE_API FStormSyncCoreDelegates
 	 * @param FStormSyncConnectedDevice Data Holding struct with info about connected device.
 	 */
 	DECLARE_MULTICAST_DELEGATE_TwoParams(FOnServiceDiscoveryConnection, const FString&, const FStormSyncConnectedDevice&);
-	
+
 	/**
 	 * Delegate type for when a connection on storm sync network state is changing (might become unresponsive).
 	 *
@@ -150,7 +149,7 @@ struct STORMSYNCCORE_API FStormSyncCoreDelegates
 	 * @param EStormSyncConnectedDeviceState Enum state indicating if connection is active or unresponsive
 	 */
 	DECLARE_MULTICAST_DELEGATE_TwoParams(FOnServiceDiscoveryStateChange, const FString&, EStormSyncConnectedDeviceState);
-	
+
 	/**
 	 * Delegate type for when server status for a connection on storm sync network has changed (either running or stopped).
 	 *
@@ -161,13 +160,12 @@ struct STORMSYNCCORE_API FStormSyncCoreDelegates
 	 */
 	DECLARE_MULTICAST_DELEGATE_TwoParams(FOnServiceDiscoveryServerStatusChange, const FString&, bool);
 
-
 	/** Broadcasted when server module starts the endpoint */
 	DECLARE_MULTICAST_DELEGATE(FOnStormSyncServerStarted)
-	
+
 	/** Broadcasted when server module stops the endpoint */
 	DECLARE_MULTICAST_DELEGATE(FOnStormSyncServerStopped)
-	
+
 
 	/**
 	 * Delegate type for a disconnection on storm sync message bus.
@@ -180,15 +178,15 @@ struct STORMSYNCCORE_API FStormSyncCoreDelegates
 	 * @param FStormSyncConnectedDevice Data Holding struct with info about connected device.
 	 */
 	DECLARE_MULTICAST_DELEGATE_OneParam(FOnServiceDiscoveryDisconnection, const FString&);
-	
+
 	/** Called when receiving a buffer over network and used to queue up an import task */
 	static FOnRequestImportBuffer OnRequestImportBuffer;
-	
+
 	/** Called when importing an .spak file via the Archive factory (import via content browser) and used to queue up an import task */
 	static FOnRequestImportFile OnRequestImportFile;
 
 	static FOnFileImported OnFileImported;
-	
+
 	/** Called when a file is extracted from an incoming storm sync pak */
 	static FOnPakAssetExtracted OnPakAssetExtracted;
 
@@ -197,28 +195,28 @@ struct STORMSYNCCORE_API FStormSyncCoreDelegates
 	 * before the buffer is created and the connection is established.
 	 */
 	static FOnPreStartSendingBuffer OnPreStartSendingBuffer;
-	
+
 	/** Called when a tcp socket client opens a tcp stream and starts sending a buffer to a remote server */
 	static FOnStartSendingBuffer OnStartSendingBuffer;
-	
+
 	/** Called when a tcp socket client is receiving a response from a server we are sending to (including the size it receives so far) */
 	static FOnReceivingBytes OnReceivingBytes;
 
 	/** Called when a tcp socket client is receiving a response from a server we are sending to (including the size it receives so far) */
 	static FOnTransferComplete OnTransferComplete;
-	
+
 	/** Called when service discovery manager detects a new connection on Storm Sync network */
 	static FOnServiceDiscoveryConnection OnServiceDiscoveryConnection;
-	
+
 	/**
 	 * Called when service discovery manager detects a state change for a connected device on
 	 * Storm Sync network (might become unresponsive)
 	 */
 	static FOnServiceDiscoveryStateChange OnServiceDiscoveryStateChange;
-	
+
 	/** Called when service discovery manager detects a state change for remote server endpoint on a connected device */
 	static FOnServiceDiscoveryServerStatusChange OnServiceDiscoveryServerStatusChange;
-	
+
 	/**
 	 * Called when service discovery manager detects a connection dropped on Storm Sync network
 	 * (happens after a certain amount of inactivity configured via StormSyncTransportSettings)
@@ -227,7 +225,7 @@ struct STORMSYNCCORE_API FStormSyncCoreDelegates
 
 	/* Server started delegate */
 	static FOnStormSyncServerStarted OnStormSyncServerStarted;
-	
+
 	/* Server stopped delegate */
 	static FOnStormSyncServerStopped OnStormSyncServerStopped;
 };
diff --git a/Engine/Plugins/Experimental/StormSync/Source/StormSyncCore/Public/StormSyncCoreLog.h b/Engine/Plugins/Experimental/StormSync/Source/StormSyncCore/Public/StormSyncCoreLog.h
deleted file mode 100644
diff --git a/Engine/Plugins/Experimental/StormSync/Source/StormSyncCore/Public/StormSyncCoreModule.h b/Engine/Plugins/Experimental/StormSync/Source/StormSyncCore/Public/StormSyncCoreModule.h
deleted file mode 100644
diff --git a/Engine/Plugins/Experimental/StormSync/Source/StormSyncCore/Public/StormSyncCoreSettings.h b/Engine/Plugins/Experimental/StormSync/Source/StormSyncCore/Public/StormSyncCoreSettings.h
--- a/Engine/Plugins/Experimental/StormSync/Source/StormSyncCore/Public/StormSyncCoreSettings.h
+++ b/Engine/Plugins/Experimental/StormSync/Source/StormSyncCore/Public/StormSyncCoreSettings.h
@@ -2,7 +2,8 @@
 
 #pragma once
 
-#include "CoreMinimal.h"
+#include "Containers/Set.h"
+#include "Containers/UnrealString.h"
 #include "Engine/DeveloperSettings.h"
 #include "StormSyncCoreSettings.generated.h"
 
@@ -11,28 +12,14 @@
  *
  * Handle configuration for buffer creation / extraction.
  */
-UCLASS(config=Game, defaultconfig)
-class STORMSYNCCORE_API UStormSyncCoreSettings : public UDeveloperSettings
+UCLASS(MinimalAPI, Config=Game, DefaultConfig, DisplayName="Core Settings")
+class UStormSyncCoreSettings : public UDeveloperSettings
 {
 	GENERATED_BODY()
 
 public:
 	UStormSyncCoreSettings();
 
-	/**
-	 * The category name for our developer settings
-	 *
-	 * @see GetCategoryName
-	 */
-	static constexpr const TCHAR* PluginCategoryName = TEXT("Storm Sync");
-	
-	//~ Begin UDeveloperSettings interface
-	virtual FName GetCategoryName() const override;
-#if WITH_EDITOR
-	virtual FText GetSectionText() const override;
-#endif
-	//~ End UDeveloperSettings interface
-
 	/**
 	 * Whether StormSync should consider packages only within the `/Game` project content when performing
 	 * an export or sending over the network.
diff --git a/Engine/Plugins/Experimental/StormSync/Source/StormSyncCore/Public/StormSyncCoreUtils.h b/Engine/Plugins/Experimental/StormSync/Source/StormSyncCore/Public/StormSyncCoreUtils.h
--- a/Engine/Plugins/Experimental/StormSync/Source/StormSyncCore/Public/StormSyncCoreUtils.h
+++ b/Engine/Plugins/Experimental/StormSync/Source/StormSyncCore/Public/StormSyncCoreUtils.h
@@ -3,7 +3,7 @@
 #pragma once
 
 #include "Async/Async.h"
-#include "CoreMinimal.h"
+#include "Delegates/Delegate.h"
 #include "StormSyncCommonTypes.h"
 #include "StormSyncPackageDescriptor.h"
 
@@ -13,7 +13,7 @@
  * Delegates used by the main extract method. Use it to register delegates when you want to be notified
  * from external code.
  */
-struct STORMSYNCCORE_API FStormSyncCoreExtractArgs
+struct FStormSyncCoreExtractArgs
 {
 	/** Delegate type for when one an incoming pak extraction process starts */
 	DECLARE_DELEGATE_OneParam(FOnPakPreExtract, int32);
@@ -52,7 +52,7 @@ struct STORMSYNCCORE_API FStormSyncCoreExtractArgs
 };
 
 /** This class exposes a set of static methods to manipulate Ava Pak files, creation and extraction. */
-class STORMSYNCCORE_API FStormSyncCoreUtils
+class FStormSyncCoreUtils
 {
 public:
 	/** Delegate type triggered on pak buffer creation for each added file */
@@ -69,8 +69,8 @@ public:
 	 * 
 	 * @return Whether the asset data could be retrieved or not
 	 */
-	static bool GetAssetData(const FString& InPackageName, TArray<FAssetData>& OutAssets, TArray<FName>& OutDependencies);
-	
+	STORMSYNCCORE_API static bool GetAssetData(const FString& InPackageName, TArray<FAssetData>& OutAssets, TArray<FName>& OutDependencies);
+
 	/**
 	 * Gathers a list of PackageNames dependencies referenced by the supplied packages.
 	 *
@@ -85,7 +85,7 @@ public:
 	 * @return Returns true if we were able to generate a list of dependencies and false if the operation failed.
 	 * Note that in case of failure, ErrorText is filled with further information.
 	 */
-	static bool GetDependenciesForPackages(const TArray<FName>& InPackageNames, TArray<FName>& OutDependencies, FText& OutErrorText, const bool bInShouldValidatePackages = true);
+	STORMSYNCCORE_API static bool GetDependenciesForPackages(const TArray<FName>& InPackageNames, TArray<FName>& OutDependencies, FText& OutErrorText, const bool bInShouldValidatePackages = true);
 
 	/**
 	 * Similar to GetDependenciesForPackages internally searching for any recursive dependencies.
@@ -101,7 +101,7 @@ public:
 	 * @return Returns true if we were able to generate a list of dependencies and false if the operation failed.
 	 * Note that in case of failure, ErrorText is filled with further information.
 	 */
-	static bool GetAvaFileDependenciesForPackages(const TArray<FName>& InPackageNames, TArray<FStormSyncFileDependency>& OutFileDependencies, FText& OutErrorText, const bool bInShouldValidatePackages = true);
+	STORMSYNCCORE_API static bool GetAvaFileDependenciesForPackages(const TArray<FName>& InPackageNames, TArray<FStormSyncFileDependency>& OutFileDependencies, FText& OutErrorText, const bool bInShouldValidatePackages = true);
 
 	/**
 	 * Creates a list of FStormSyncFileDependency (with info such as asset path, file size, timestamp and file hash)
@@ -111,7 +111,7 @@ public:
 	 *
 	 * @return Returns a list of FStormSyncFileDependency
 	 */
-	static TArray<FStormSyncFileDependency> GetAvaFileDependenciesFromPackageNames(const TArray<FName>& InPackageNames, bool bInShouldIncludeInvalid = false);
+	STORMSYNCCORE_API static TArray<FStormSyncFileDependency> GetAvaFileDependenciesFromPackageNames(const TArray<FName>& InPackageNames, bool bInShouldIncludeInvalid = false);
 
 	/**
 	 * Async version of FStormSyncCoreUtils::GetAvaFileDependenciesFromPackageNames to return a list of
@@ -123,14 +123,14 @@ public:
 	 * @param InThreadType Optional Thread Type to use for the async execution (Default EAsyncExecution::Thread)
 	 * @return A future containing the list of file dependencies.
 	 */
-	static TFuture<TArray<FStormSyncFileDependency>> GetAvaFileDependenciesAsync(const TArray<FName>& InPackageNames, const bool bInShouldValidatePackages = true, EAsyncExecution InThreadType = EAsyncExecution::Thread);
+	STORMSYNCCORE_API static TFuture<TArray<FStormSyncFileDependency>> GetAvaFileDependenciesAsync(const TArray<FName>& InPackageNames, const bool bInShouldValidatePackages = true, EAsyncExecution InThreadType = EAsyncExecution::Thread);
 
 	/**
 	 * Creates a new FStormSyncFileDependency and initialize filesystem related information
 	 *
 	 * Note: Considering renaming FStormSyncFileDependency to FStormSyncFileInfo (to distinguish with file references / dependencies idea)
 	 */
-	static FStormSyncFileDependency CreateStormSyncFile(const FName& InPackageName);
+	STORMSYNCCORE_API static FStormSyncFileDependency CreateStormSyncFile(const FName& InPackageName);
 
 	/**
 	 * Recursively get dependencies for provided list of Package Names and attempts to create an ava pak buffer using a
@@ -149,7 +149,7 @@ public:
 	 * @return Returns true if we were able to generate the pak buffer false if the operation failed.
 	 * Note that in case of failure, ErrorText is filled with further information.
 	 */
-	static bool CreatePakBufferWithDependencies(const TArray<FName>& InPackageNames, TArray<uint8>& OutPakBuffer, FText& OutErrorText, const FOnFileAdded& InOnFileAdded = FOnFileAdded());
+	STORMSYNCCORE_API static bool CreatePakBufferWithDependencies(const TArray<FName>& InPackageNames, TArray<uint8>& OutPakBuffer, FText& OutErrorText, const FOnFileAdded& InOnFileAdded = FOnFileAdded());
 
 	/**
 	 * Attempts to create an ava pak buffer using a memory archive from the provided list of Package Names.
@@ -170,7 +170,7 @@ public:
 	 * @return Returns true if we were able to generate the pak buffer false if the operation failed.
 	 * Note that in case of failure, ErrorText is filled with further information.
 	 */
-	static bool CreatePakBuffer(const TArray<FName>& InPackageNames, TArray<uint8>& OutPakBuffer, FText& OutErrorText, const FOnFileAdded& InOnFileAdded = FOnFileAdded());
+	STORMSYNCCORE_API static bool CreatePakBuffer(const TArray<FName>& InPackageNames, TArray<uint8>& OutPakBuffer, FText& OutErrorText, const FOnFileAdded& InOnFileAdded = FOnFileAdded());
 
 	/**
 	 * Main entry point for unpacking. Extracts an ava pak into the current Unreal project (Game content or Plugins content folder)
@@ -185,7 +185,7 @@ public:
 	 * @return Returns true if we were able to fully import the package and false if the operation failed (even for partial imports).
 	 * Note that in case of failure, Errors array of localized text is filled with further information.
 	 */
-	static bool ExtractPakBuffer(const TArray<uint8>& InPakBuffer, const FStormSyncCoreExtractArgs& InExtractArgs, TMap<FString, FString>& OutSuccessfullyExtractedPackages, TArray<FText>& OutErrors);
+	STORMSYNCCORE_API static bool ExtractPakBuffer(const TArray<uint8>& InPakBuffer, const FStormSyncCoreExtractArgs& InExtractArgs, TMap<FString, FString>& OutSuccessfullyExtractedPackages, TArray<FText>& OutErrors);
 
 	/**
 	 * Converts the provided bytes size into a human readable string.
@@ -194,8 +194,8 @@ public:
 	 *
 	 * @return A string representing the provided size with B, KB, MB of GB suffix.
 	 */
-	static FString GetHumanReadableByteSize(uint64 InSize);
-	
+	STORMSYNCCORE_API static FString GetHumanReadableByteSize(uint64 InSize);
+
 	/**
 	 * Figures out the files that needs to be synced based on local and remote list of dependencies.
 	 *
@@ -215,8 +215,8 @@ public:
 	 *
 	 * @return List of FStormSyncFileModifierInfo indicating files that are either missing or dirty
 	 */
-	static TArray<FStormSyncFileModifierInfo> GetSyncFileModifiers(const TArray<FName>& InPackageNames, const TArray<FStormSyncFileDependency>& InRemoteDependencies);
-	
+	STORMSYNCCORE_API static TArray<FStormSyncFileModifierInfo> GetSyncFileModifiers(const TArray<FName>& InPackageNames, const TArray<FStormSyncFileDependency>& InRemoteDependencies);
+
 private:
 	/**
 	 * Gets the dependencies of the specified package recursively.
diff --git a/Engine/Plugins/Experimental/StormSync/Source/StormSyncCore/Public/StormSyncImportTypes.h b/Engine/Plugins/Experimental/StormSync/Source/StormSyncCore/Public/StormSyncImportTypes.h
--- a/Engine/Plugins/Experimental/StormSync/Source/StormSyncCore/Public/StormSyncImportTypes.h
+++ b/Engine/Plugins/Experimental/StormSync/Source/StormSyncCore/Public/StormSyncImportTypes.h
@@ -2,7 +2,8 @@
 
 #pragma once
 
-#include "CoreMinimal.h"
+#include "Containers/UnrealString.h"
+#include "Internationalization/Text.h"
 #include "StormSyncCommonTypes.h"
 #include "StormSyncPackageDescriptor.h"
 #include "StormSyncImportTypes.generated.h"
@@ -13,7 +14,7 @@
  * Used to compute a list of files that needs extraction.
  */
 USTRUCT()
-struct STORMSYNCCORE_API FStormSyncImportFileInfo
+struct FStormSyncImportFileInfo
 {
 	GENERATED_BODY()
 	
diff --git a/Engine/Plugins/Experimental/StormSync/Source/StormSyncCore/Public/StormSyncPackageDescriptor.h b/Engine/Plugins/Experimental/StormSync/Source/StormSyncCore/Public/StormSyncPackageDescriptor.h
--- a/Engine/Plugins/Experimental/StormSync/Source/StormSyncCore/Public/StormSyncPackageDescriptor.h
+++ b/Engine/Plugins/Experimental/StormSync/Source/StormSyncCore/Public/StormSyncPackageDescriptor.h
@@ -2,7 +2,6 @@
 
 #pragma once
 
-#include "CoreMinimal.h"
 #include "StormSyncPackageDescriptor.generated.h"
 
 /**
@@ -13,7 +12,7 @@
  * for a given ava package.
  */
 USTRUCT()
-struct STORMSYNCCORE_API FStormSyncFileDependency
+struct FStormSyncFileDependency
 {
 	GENERATED_BODY()
 
@@ -65,12 +64,12 @@ struct STORMSYNCCORE_API FStormSyncFileDependency
 	 * @return Destination filepath or empty string if we were not able to determine destination output
 	 * (eg, Plugin package while plugin not installed or enabled in current project)
 	 */
-	static FString GetDestFilepath(const FName& InPackageName, FText& OutFailureReason);
+	STORMSYNCCORE_API static FString GetDestFilepath(const FName& InPackageName, FText& OutFailureReason);
 
 	/** Returns whether file info was properly initialized */
-	bool IsValid() const;
+	STORMSYNCCORE_API bool IsValid() const;
 
-	FString ToString() const;
+	STORMSYNCCORE_API FString ToString() const;
 };
 
 /**
@@ -81,7 +80,7 @@ struct STORMSYNCCORE_API FStormSyncFileDependency
  * their state, timestamp, size, hash, etc.)
  */
 USTRUCT()
-struct STORMSYNCCORE_API FStormSyncPackageDescriptor
+struct FStormSyncPackageDescriptor
 {
 	GENERATED_BODY()
 
@@ -116,7 +115,7 @@ struct STORMSYNCCORE_API FStormSyncPackageDescriptor
 	UPROPERTY(VisibleAnywhere, Category = "Storm Sync")
 	TArray<FStormSyncFileDependency> Dependencies;
 
-	FString ToString() const;
+	STORMSYNCCORE_API FString ToString() const;
 };
 
 /**
diff --git a/Engine/Plugins/Experimental/StormSync/Source/StormSyncDrives/Private/StormSyncDrivesLog.h b/Engine/Plugins/Experimental/StormSync/Source/StormSyncDrives/Private/StormSyncDrivesLog.h
new file mode 100644
--- /dev/null
+++ b/Engine/Plugins/Experimental/StormSync/Source/StormSyncDrives/Private/StormSyncDrivesLog.h
@@ -0,0 +1,12 @@
+// Copyright Epic Games, Inc. All Rights Reserved.
+
+#pragma once
+
+#include "Logging/LogMacros.h"
+
+DECLARE_LOG_CATEGORY_EXTERN(LogStormSyncDrives, Display, All);
+
+#define STORM_SYNC_DRIVES_LOG(Verbosity, Format, ...) \
+{ \
+    UE_LOG(LogStormSyncDrives, Verbosity, Format, ##__VA_ARGS__); \
+}
diff --git a/Engine/Plugins/Experimental/StormSync/Source/StormSyncDrives/Private/StormSyncDrivesModule.h b/Engine/Plugins/Experimental/StormSync/Source/StormSyncDrives/Private/StormSyncDrivesModule.h
new file mode 100644
--- /dev/null
+++ b/Engine/Plugins/Experimental/StormSync/Source/StormSyncDrives/Private/StormSyncDrivesModule.h
@@ -0,0 +1,84 @@
+// Copyright Epic Games, Inc. All Rights Reserved.
+
+#pragma once
+
+#include "IStormSyncDrivesModule.h"
+#include "UObject/UnrealType.h"
+
+#if WITH_EDITOR
+#include "IMessageLogListing.h"
+#endif
+
+struct FStormSyncMountPointConfig;
+class UStormSyncDrivesSettings;
+
+/** Main entry point and implementation of StormSync Mounted Drives Runtime module. */
+class FStormSyncDrivesModule : public IStormSyncDrivesModule
+{
+public:
+	//~ Begin IModuleInterface interface
+	virtual void StartupModule() override;
+	virtual void ShutdownModule() override;
+	//~ End IModuleInterface interface
+
+	//~ Begin IStormSyncDrivesModule interface
+	virtual bool RegisterMountPoint(const FStormSyncMountPointConfig& InMountPoint, FText& ErrorText) override;
+	virtual bool UnregisterMountPoint(const FStormSyncMountPointConfig& InMountPoint, FText& ErrorText) override;
+	//~ End IStormSyncDrivesModule interface
+
+private:
+	/** The name of the message output log we will send messages to */
+	static constexpr const TCHAR* LogName = TEXT("StormSyncDrives");
+	
+	/**
+	 * Array of Mount Points templates populated from UStormSyncDrivesSettings.
+	 * 
+	 * Allows projects to specify reusable mount points templates for the content browser.
+	*/
+	TArray<TSharedPtr<FStormSyncMountPointConfig>> MountedDrives;
+
+#if WITH_EDITORONLY_DATA
+	/** Listing used in the editor by the Message Log. */
+	TSharedPtr<IMessageLogListing> LogListing;
+#endif
+
+	/** Core OnEngineLoopInitComplete delegate used to mount any stored config in developer settings */
+	void OnEngineLoopInitComplete();
+
+	/**
+	 * Developer settings changed handler. Called anytime settings changed.
+	 *
+	 * Used to validate user config and if successful, mount the set of Mount Points to Mount Directories.
+	 */
+	void OnSettingsChanged(UObject* InSettings, FPropertyChangedEvent& InPropertyChangedEvent);
+
+	/** Performs validation for a single mount point config and log any errors to the message log */
+	bool ValidateMountPoint(const FStormSyncMountPointConfig& InMountPoint, int32 Index) const;
+	
+	/**
+	 * Performs validation for the provided list of mount point config. Logs any errors to the message log.
+	 *
+	 * Happens once on engine init, and anytime the user settings changed.
+	 */
+	bool ValidateMountPoints(const TArray<FStormSyncMountPointConfig>& InMountPoints) const;
+
+	/**
+	 * Unregisters previously mounted drives, cache (update internal list of MountedDrives shared pointers),
+	 * and register / mount the new config.
+	 * 
+	 * Happens once on engine init, and anytime the user settings changed.
+	 */
+	void ResetMountedDrivesFromSettings(const UStormSyncDrivesSettings* InSettings);
+
+	/** Unregisters any previously mounted drives */
+	void UnregisterMountedDrives();
+
+	/** Reset and updates internal list of registered drives */
+	void CacheMountedDrives(const TArray<FStormSyncMountPointConfig>& InMountPoints);
+
+	/** Tries to mount previously cached list of mount point config */
+	void RegisterMountedDrives();
+
+	/** Static helper to log an error to the message log. Used during validation process. */
+	static void AddMessageError(const FText& Text);
+};
diff --git a/Engine/Plugins/Experimental/StormSync/Source/StormSyncDrives/Private/StormSyncDrivesSettings.cpp b/Engine/Plugins/Experimental/StormSync/Source/StormSyncDrives/Private/StormSyncDrivesSettings.cpp
--- a/Engine/Plugins/Experimental/StormSync/Source/StormSyncDrives/Private/StormSyncDrivesSettings.cpp
+++ b/Engine/Plugins/Experimental/StormSync/Source/StormSyncDrives/Private/StormSyncDrivesSettings.cpp
@@ -1,20 +1,8 @@
 // Copyright Epic Games, Inc. All Rights Reserved.
 
-
 #include "StormSyncDrivesSettings.h"
 
 UStormSyncDrivesSettings::UStormSyncDrivesSettings()
 {
+	CategoryName = TEXT("Storm Sync");
 }
-
-FName UStormSyncDrivesSettings::GetCategoryName() const
-{
-	return PluginCategoryName;
-}
-
-#if WITH_EDITOR
-FText UStormSyncDrivesSettings::GetSectionText() const
-{
-	return NSLOCTEXT("StormSyncDrives", "StormSyncDrivesSettingsSection", "Mount Points Settings");
-}
-#endif
diff --git a/Engine/Plugins/Experimental/StormSync/Source/StormSyncDrives/Public/IStormSyncDrivesModule.h b/Engine/Plugins/Experimental/StormSync/Source/StormSyncDrives/Public/IStormSyncDrivesModule.h
--- a/Engine/Plugins/Experimental/StormSync/Source/StormSyncDrives/Public/IStormSyncDrivesModule.h
+++ b/Engine/Plugins/Experimental/StormSync/Source/StormSyncDrives/Public/IStormSyncDrivesModule.h
@@ -2,14 +2,13 @@
 
 #pragma once
 
-#include "CoreMinimal.h"
 #include "Modules/ModuleInterface.h"
 #include "Modules/ModuleManager.h"
 
 struct FStormSyncMountPointConfig;
 
 /** Public interface for Storm Sync Drives module and API */
-class STORMSYNCDRIVES_API IStormSyncDrivesModule : public IModuleInterface
+class IStormSyncDrivesModule : public IModuleInterface
 {
 public:
 	static bool IsAvailable()
diff --git a/Engine/Plugins/Experimental/StormSync/Source/StormSyncDrives/Public/StormSyncDrivesLog.h b/Engine/Plugins/Experimental/StormSync/Source/StormSyncDrives/Public/StormSyncDrivesLog.h
deleted file mode 100644
diff --git a/Engine/Plugins/Experimental/StormSync/Source/StormSyncDrives/Public/StormSyncDrivesModule.h b/Engine/Plugins/Experimental/StormSync/Source/StormSyncDrives/Public/StormSyncDrivesModule.h
deleted file mode 100644
diff --git a/Engine/Plugins/Experimental/StormSync/Source/StormSyncDrives/Public/StormSyncDrivesSettings.h b/Engine/Plugins/Experimental/StormSync/Source/StormSyncDrives/Public/StormSyncDrivesSettings.h
--- a/Engine/Plugins/Experimental/StormSync/Source/StormSyncDrives/Public/StormSyncDrivesSettings.h
+++ b/Engine/Plugins/Experimental/StormSync/Source/StormSyncDrives/Public/StormSyncDrivesSettings.h
@@ -2,13 +2,12 @@
 
 #pragma once
 
-#include "CoreMinimal.h"
 #include "Engine/DeveloperSettings.h"
 #include "StormSyncDrivesSettings.generated.h"
 
 /** Holds config for a given Mount Point (eg. /RootPath) to a filesystem directory (can be local, could be networked) */
 USTRUCT()
-struct STORMSYNCDRIVES_API FStormSyncMountPointConfig
+struct FStormSyncMountPointConfig
 {
 	GENERATED_BODY()
 
@@ -53,29 +52,15 @@ struct STORMSYNCDRIVES_API FStormSyncMountPointConfig
  *
  * If successful, mount points will be updated and should be available in the content browser.
  */
-UCLASS(config=Game, defaultconfig)
-class STORMSYNCDRIVES_API UStormSyncDrivesSettings : public UDeveloperSettings
+UCLASS(MinimalAPI, Config=Game, DefaultConfig, DisplayName="Mount Points Settings")
+class UStormSyncDrivesSettings : public UDeveloperSettings
 {
 	GENERATED_BODY()
-	
+
 public:
 	UStormSyncDrivesSettings();
 
-	/**
-	 * The category name for our developer settings
-	 *
-	 * @see GetCategoryName
-	 */
-	static constexpr const TCHAR* PluginCategoryName = TEXT("Storm Sync");
-
 	/** List of mounted drives configuration */
 	UPROPERTY(config, EditAnywhere, Category = "Mounted Drives", meta=(TitleProperty="{MountPoint} mapped to {MountDirectory}"))
 	TArray<FStormSyncMountPointConfig> MountPoints;
-	
-	//~ Begin UDeveloperSettings interface
-	virtual FName GetCategoryName() const override;
-#if WITH_EDITOR
-	virtual FText GetSectionText() const override;
-#endif
-	//~ End UDeveloperSettings interface
 };
diff --git a/Engine/Plugins/Experimental/StormSync/Source/StormSyncDrives/Public/StormSyncDrivesUtils.h b/Engine/Plugins/Experimental/StormSync/Source/StormSyncDrives/Public/StormSyncDrivesUtils.h
--- a/Engine/Plugins/Experimental/StormSync/Source/StormSyncDrives/Public/StormSyncDrivesUtils.h
+++ b/Engine/Plugins/Experimental/StormSync/Source/StormSyncDrives/Public/StormSyncDrivesUtils.h
@@ -2,7 +2,6 @@
 
 #pragma once
 
-#include "CoreMinimal.h"
 #include "Engine/EngineTypes.h"
 #include "Misc/PackageName.h"
 
diff --git a/Engine/Plugins/Experimental/StormSync/Source/StormSyncEditor/Private/ContextMenu/StormSyncAssetFolderContextMenu.h b/Engine/Plugins/Experimental/StormSync/Source/StormSyncEditor/Private/ContextMenu/StormSyncAssetFolderContextMenu.h
--- a/Engine/Plugins/Experimental/StormSync/Source/StormSyncEditor/Private/ContextMenu/StormSyncAssetFolderContextMenu.h
+++ b/Engine/Plugins/Experimental/StormSync/Source/StormSyncEditor/Private/ContextMenu/StormSyncAssetFolderContextMenu.h
@@ -2,7 +2,7 @@
 
 #pragma once
 
-#include "CoreMinimal.h"
+#include "Templates/SharedPointer.h"
 
 class FMenuBuilder;
 class UToolMenu;
diff --git a/Engine/Plugins/Experimental/StormSync/Source/StormSyncEditor/Private/Customization/SStormSyncIpAddressComboBox.h b/Engine/Plugins/Experimental/StormSync/Source/StormSyncEditor/Private/Customization/SStormSyncIpAddressComboBox.h
--- a/Engine/Plugins/Experimental/StormSync/Source/StormSyncEditor/Private/Customization/SStormSyncIpAddressComboBox.h
+++ b/Engine/Plugins/Experimental/StormSync/Source/StormSyncEditor/Private/Customization/SStormSyncIpAddressComboBox.h
@@ -2,8 +2,6 @@
 
 #pragma once
 
-#include "CoreMinimal.h"
-
 #include "Widgets/SCompoundWidget.h"
 #include "Widgets/Input/SComboBox.h"
 #include "Widgets/Input/SEditableText.h"
diff --git a/Engine/Plugins/Experimental/StormSync/Source/StormSyncEditor/Private/Customization/StormSyncTransportSettingsDetailsCustomization.h b/Engine/Plugins/Experimental/StormSync/Source/StormSyncEditor/Private/Customization/StormSyncTransportSettingsDetailsCustomization.h
--- a/Engine/Plugins/Experimental/StormSync/Source/StormSyncEditor/Private/Customization/StormSyncTransportSettingsDetailsCustomization.h
+++ b/Engine/Plugins/Experimental/StormSync/Source/StormSyncEditor/Private/Customization/StormSyncTransportSettingsDetailsCustomization.h
@@ -2,7 +2,6 @@
 
 #pragma once
 
-#include "CoreMinimal.h"
 #include "IDetailCustomization.h"
 
 class IPropertyHandle;
diff --git a/Engine/Plugins/Experimental/StormSync/Source/StormSyncEditor/Private/Factories/StormSyncArchiveFactory.h b/Engine/Plugins/Experimental/StormSync/Source/StormSyncEditor/Private/Factories/StormSyncArchiveFactory.h
--- a/Engine/Plugins/Experimental/StormSync/Source/StormSyncEditor/Private/Factories/StormSyncArchiveFactory.h
+++ b/Engine/Plugins/Experimental/StormSync/Source/StormSyncEditor/Private/Factories/StormSyncArchiveFactory.h
@@ -2,7 +2,6 @@
 
 #pragma once
 
-#include "CoreMinimal.h"
 #include "Factories/Factory.h"
 #include "StormSyncArchiveFactory.generated.h"
 
diff --git a/Engine/Plugins/Experimental/StormSync/Source/StormSyncEditor/Private/Slate/SStormSyncExportWizard.h b/Engine/Plugins/Experimental/StormSync/Source/StormSyncEditor/Private/Slate/SStormSyncExportWizard.h
--- a/Engine/Plugins/Experimental/StormSync/Source/StormSyncEditor/Private/Slate/SStormSyncExportWizard.h
+++ b/Engine/Plugins/Experimental/StormSync/Source/StormSyncEditor/Private/Slate/SStormSyncExportWizard.h
@@ -2,7 +2,6 @@
 
 #pragma once
 
-#include "CoreMinimal.h"
 #include "Widgets/Views/SListView.h"
 #include "Widgets/Views/SHeaderRow.h"
 #include "Widgets/SCompoundWidget.h"
diff --git a/Engine/Plugins/Experimental/StormSync/Source/StormSyncEditor/Private/Slate/SStormSyncFileDependencyWidgetRow.h b/Engine/Plugins/Experimental/StormSync/Source/StormSyncEditor/Private/Slate/SStormSyncFileDependencyWidgetRow.h
new file mode 100644
--- /dev/null
+++ b/Engine/Plugins/Experimental/StormSync/Source/StormSyncEditor/Private/Slate/SStormSyncFileDependencyWidgetRow.h
@@ -0,0 +1,44 @@
+// Copyright Epic Games, Inc. All Rights Reserved.
+
+#pragma once
+
+#include "Widgets/Views/SListView.h"
+#include "StormSyncPackageDescriptor.h"
+
+struct FStormSyncImportFileInfo;
+
+namespace StormSync::SlateWidgetRow
+{
+	static FName HeaderRow_PackageName = TEXT("PackageName");
+	static FName HeaderRow_FileSize = TEXT("FileSize");
+	static FName HeaderRow_Timestamp = TEXT("Timestamp");
+	static FName HeaderRow_FileHash = TEXT("FileHash");
+	static FName HeaderRow_ImportReason = TEXT("ImportReason");
+
+	static FText DefaultLabel_PackageName = NSLOCTEXT("StormSyncFileDependencyWidgetRow", "PackageNameColumnHeader", "Package Name");
+	static FText DefaultLabel_FileSize = NSLOCTEXT("StormSyncFileDependencyWidgetRow", "FileSizeColumnHeader", "File Size");
+	static FText DefaultLabel_Timestamp = NSLOCTEXT("StormSyncFileDependencyWidgetRow", "TimestampColumnHeader", "Timestamp");
+	static FText DefaultLabel_FileHash = NSLOCTEXT("StormSyncFileDependencyWidgetRow", "FileHashColumnHeader", "File Hash");
+	static FText DefaultLabel_ImportReason = NSLOCTEXT("StormSyncFileDependencyWidgetRow", "ImportReasonColumnHeader", "Import Reason");
+}
+
+/** Table Row for StormSyncFileDependency list */
+class SStormSyncFileDependencyWidgetRow : public SMultiColumnTableRow<TSharedPtr<FStormSyncFileDependency>>
+{
+public:
+	SLATE_BEGIN_ARGS(SStormSyncFileDependencyWidgetRow) {}
+	
+	SLATE_ARGUMENT(TSharedPtr<FStormSyncImportFileInfo>, Item)
+
+	SLATE_END_ARGS()
+
+	/** Constructs this widget with InArgs and OwnerTable */
+	void Construct(const FArguments& InArgs, const TSharedRef<STableViewBase>& InOwnerTableView);
+
+	/** Generates appropriate widget for the given column name ID */
+	virtual TSharedRef<SWidget> GenerateWidgetForColumn(const FName& InColumnName) override;
+
+private:
+	/** The underlying item that the TableView is visualizing */
+	TSharedPtr<FStormSyncImportFileInfo> WidgetItem;
+};
diff --git a/Engine/Plugins/Experimental/StormSync/Source/StormSyncEditor/Private/Slate/SStormSyncImportWizard.cpp b/Engine/Plugins/Experimental/StormSync/Source/StormSyncEditor/Private/Slate/SStormSyncImportWizard.cpp
--- a/Engine/Plugins/Experimental/StormSync/Source/StormSyncEditor/Private/Slate/SStormSyncImportWizard.cpp
+++ b/Engine/Plugins/Experimental/StormSync/Source/StormSyncEditor/Private/Slate/SStormSyncImportWizard.cpp
@@ -177,38 +177,11 @@ void SStormSyncImportWizard::InitListSources()
 	}
 }
 
-TSharedRef<SStormSyncImportWizard> SStormSyncImportWizard::CreateWizard(const TArray<FStormSyncImportFileInfo>& InFilesToImport, const TArray<FStormSyncImportFileInfo>& InBufferFiles)
-{
-	const TSharedRef<SStormSyncImportWizard> ImportWizard = SNew(SStormSyncImportWizard, InFilesToImport, InBufferFiles);
-	
-	const TSharedRef<SWindow> Window = SNew(SWindow)
-	.Title(LOCTEXT("Window_Title", "Storm Sync | Import files from local pak"))
-	.ClientSize(FVector2D(960, 700))
-	.SupportsMaximize(false)
-	.SupportsMinimize(false)
-	[
-		ImportWizard->AsShared()
-	];
-	
-	TSharedPtr<SWindow> ParentWindow;
-
-	if (FModuleManager::Get().IsModuleLoaded("MainFrame"))
-	{
-		const IMainFrameModule& MainFrame = FModuleManager::LoadModuleChecked<IMainFrameModule>("MainFrame");
-		ParentWindow = MainFrame.GetParentWindow();
-	}
-
-	FSlateApplication::Get().AddModalWindow(Window, ParentWindow, false);
-
-	return ImportWizard;
-}
-
 bool SStormSyncImportWizard::ShouldImport() const
 {
 	return bShouldImport;
 }
 
-
 TSharedRef<ITableRow> SStormSyncImportWizard::MakeFileDependencyWidget(const TSharedPtr<FStormSyncImportFileInfo> InItem, const TSharedRef<STableViewBase>& OwnerTable) const
 {
 	return SNew(SStormSyncFileDependencyWidgetRow, OwnerTable).Item(InItem);
diff --git a/Engine/Plugins/Experimental/StormSync/Source/StormSyncEditor/Private/Slate/SStormSyncImportWizard.h b/Engine/Plugins/Experimental/StormSync/Source/StormSyncEditor/Private/Slate/SStormSyncImportWizard.h
new file mode 100644
--- /dev/null
+++ b/Engine/Plugins/Experimental/StormSync/Source/StormSyncEditor/Private/Slate/SStormSyncImportWizard.h
@@ -0,0 +1,109 @@
+// Copyright Epic Games, Inc. All Rights Reserved.
+
+#pragma once
+
+#include "IStormSyncImportWizard.h"
+#include "StormSyncImportTypes.h"
+#include "Widgets/SCompoundWidget.h"
+#include "Widgets/Views/SListView.h"
+
+class SWizard;
+class SWidgetSwitcher;
+
+/** Which tab page are we currently displaying */
+enum class EStormSyncImportWizardActiveTab : uint8
+{
+	FilesToImport,
+	BufferFiles
+};
+
+/**
+ * A modal dialog to collect information about files to import in local project.
+ *
+ * It is wrapping a wizard with a single-page right now (could argue that it doesn't need to be a wizard in that case),
+ * and a tab-like view with segmented control and widget switcher to display either the list of files to import (files
+ * with mismatched state between local and buffer) or files in buffer.
+ *
+ * Note: UCommonTabListWidgetBase from CommonUI seems to implement a tab-like widget that would fit perfectly for that use case.
+ */
+class SStormSyncImportWizard : public SCompoundWidget, public IStormSyncImportWizard
+{
+public:
+	SLATE_BEGIN_ARGS(SStormSyncImportWizard) {}
+	SLATE_END_ARGS()
+
+	/** Constructs this widget with InArgs, list of FStormSyncImportFileInfo to import and full list of FStormSyncImportFileInfo from buffer */
+	void Construct(const FArguments& InArgs, const TArray<FStormSyncImportFileInfo>& InFilesToImport, const TArray<FStormSyncImportFileInfo>& InBufferFiles);
+
+protected:
+	//~ Begin IStormSyncImportWizard
+	virtual bool ShouldImport() const override;
+	//~ End IStormSyncImportWizard
+
+private:
+	/** Main wizard */
+	TSharedPtr<SWizard> Wizard;
+	
+	/** The packages list view in files to import tab */
+	TSharedPtr<SListView<TSharedPtr<FStormSyncImportFileInfo>>> ListViewFilesToImport;
+	
+	/** The packages list view in files to import tab */
+	TSharedPtr<SListView<TSharedPtr<FStormSyncImportFileInfo>>> ListViewBufferFiles;
+	
+	/** File to import list view data source */
+	TArray<TSharedPtr<FStormSyncImportFileInfo>> FilesToImportListSource;
+	
+	/** File to import list view data source */
+	TArray<TSharedPtr<FStormSyncImportFileInfo>> BufferFilesListSource;
+
+	/** Original list of file infos for files to import (mismatched state between buffer and local files) */
+	TArray<FStormSyncImportFileInfo> FilesToImport;
+	
+	/** Original list of file infos included in buffer */
+	TArray<FStormSyncImportFileInfo> BufferFiles;
+
+	/** Current value for the segmented control active view */
+	TAttribute<EStormSyncImportWizardActiveTab> CurrentTab;
+
+	/** Switcher for active tab view */
+	TSharedPtr<SWidgetSwitcher> TabContentSwitcher;
+
+	/** Holds user choice on whether to perform import or cancel */
+	bool bShouldImport = false;
+
+	/** Last column id that has been sorted */
+	FName ColumnIdToSort;
+
+	/** Current active sort mode */
+	EColumnSortMode::Type ActiveSortMode = EColumnSortMode::Type::None;
+
+	/** Initializes sources for the list views from FilesToImport / BufferFiles */
+	void InitListSources();
+	
+	/** Generate table row for file dependency list view */
+	TSharedRef<ITableRow> MakeFileDependencyWidget(TSharedPtr<FStormSyncImportFileInfo> InItem, const TSharedRef<STableViewBase>& OwnerTable) const;
+	
+	/** Handler for cancel button, which closes the dialog */
+	void OnCancelButtonClicked();
+	
+	/** Handler for finish wizard button, which closes the dialog and sets up bShouldImport to true before exiting */
+	void OnFinish();
+	
+	/** Returns whether wizard can be completed */
+	bool CanFinish() const;
+
+	/** Closes the dialog */
+	void CloseDialog();
+
+	/** Handler for segmented control tab view changed */
+	void OnTabViewChanged(EStormSyncImportWizardActiveTab InActiveTab);
+
+	/** Adjust finish button depending on wizard is able to finish (has files to import) */
+	FText GetFinishButtonTooltip() const;
+
+	/** Returns last stored sort mode for column id if it matches last column being sorted */
+	EColumnSortMode::Type GetSortModeForColumn(FName InColumnId) const;
+
+	/** Updates both ColumnIdToSort and ActiveSortMode before performing the sort */
+	void OnSortAttributeEntries(EColumnSortPriority::Type InPriority, const FName& InColumnId, EColumnSortMode::Type InSortMode);
+};
diff --git a/Engine/Plugins/Experimental/StormSync/Source/StormSyncEditor/Private/Slate/SStormSyncProgressBarNotification.h b/Engine/Plugins/Experimental/StormSync/Source/StormSyncEditor/Private/Slate/SStormSyncProgressBarNotification.h
--- a/Engine/Plugins/Experimental/StormSync/Source/StormSyncEditor/Private/Slate/SStormSyncProgressBarNotification.h
+++ b/Engine/Plugins/Experimental/StormSync/Source/StormSyncEditor/Private/Slate/SStormSyncProgressBarNotification.h
@@ -2,7 +2,6 @@
 
 #pragma once
 
-#include "CoreMinimal.h"
 #include "Widgets/SCompoundWidget.h"
 #include "Widgets/Notifications/INotificationWidget.h"
 
diff --git a/Engine/Plugins/Experimental/StormSync/Source/StormSyncEditor/Private/Slate/SStormSyncReportDialog.h b/Engine/Plugins/Experimental/StormSync/Source/StormSyncEditor/Private/Slate/SStormSyncReportDialog.h
new file mode 100644
--- /dev/null
+++ b/Engine/Plugins/Experimental/StormSync/Source/StormSyncEditor/Private/Slate/SStormSyncReportDialog.h
@@ -0,0 +1,132 @@
+// Copyright Epic Games, Inc. All Rights Reserved.
+
+#pragma once
+
+#include "Widgets/SCompoundWidget.h"
+#include "Widgets/Views/STreeView.h"
+
+struct FStormSyncPackageReportNode;
+typedef STreeView<TSharedPtr<FStormSyncPackageReportNode>> SStormSyncReportTreeView;
+
+/** Struct equivalent of ReportPackageData in Engine/Source/Developer/AssetTools/Private/SPackageReportDialog.h */
+struct FStormSyncReportPackageData
+{
+	/** Package name, eg. /Game/Path/File */
+	FString Name;
+
+	/** Whether this package name is selected in the tree view */
+	bool bShouldIncludePackage = false;
+
+	FStormSyncReportPackageData(const FString& InName, const bool bInShouldMigratePackage)
+		: Name(InName)
+		, bShouldIncludePackage(bInShouldMigratePackage)
+	{
+	}
+
+	FString ToString() const;
+};
+
+/** Struct equivalent of FPackageReportNode in Engine/Source/Developer/AssetTools/Private/SPackageReportDialog.h */
+struct FStormSyncPackageReportNode
+{
+	/** The name of the tree node without the path */
+	FString NodeName;
+
+	/** A user-exposed flag determining whether the content of this node and its children should be migrated or not. */
+	ECheckBoxState CheckedState;
+
+	/** A pointer to an external bool describing whether this node ultimately should be migrated or not. Is only non-null for leaf nodes.*/
+	bool* bShouldMigratePackage;
+
+	/** If true, this node is a folder instead of a package */
+	bool bIsFolder;
+
+	/** The parent of this node */
+	FStormSyncPackageReportNode* Parent;
+
+	/** The children of this node */
+	TArray<TSharedPtr<FStormSyncPackageReportNode>> Children;
+
+	/** Constructor */
+	FStormSyncPackageReportNode()
+		: CheckedState(ECheckBoxState::Undetermined)
+		, bShouldMigratePackage(nullptr)
+		, bIsFolder(false)
+		, Parent(nullptr)
+	{
+	}
+
+	FStormSyncPackageReportNode(const FString& InNodeName, bool InIsFolder)
+		: NodeName(InNodeName)
+		, CheckedState(ECheckBoxState::Undetermined)
+		, bShouldMigratePackage(nullptr)
+		, bIsFolder(InIsFolder)
+		, Parent(nullptr)
+	{
+	}
+
+	/** Adds the path to the tree relative to this node, creating nodes as needed. */
+	void AddPackage(const FString& PackageName, bool* bInShouldMigratePackage);
+
+	/** Expands this node and all its children */
+	void ExpandChildrenRecursively(const TSharedRef<SStormSyncReportTreeView>& InTreeView);
+
+private:
+	struct FChildrenState
+	{
+		bool bAnyChildIsChecked;
+		bool bAllChildrenAreChecked;
+	};
+
+	/** Helper function for AddPackage. PathElements is the tokenized path delimited by "/" */
+	FChildrenState AddPackage_Recursive(TArray<FString>& PathElements, bool* bInShouldMigratePackage);
+};
+
+/**
+ * Straight up copy of Engine/Source/Developer/AssetTools/Private/SPackageReportDialog.h
+ * which is used in Migrate Packages action.
+ *
+ * We can't use it directly as it's not exposed to outside modules.
+ */
+class SStormSyncReportDialog final : public SCompoundWidget
+{
+public:
+	SLATE_BEGIN_ARGS(SStormSyncReportDialog) {}
+	SLATE_END_ARGS()
+
+	/** Constructs this widget with InArgs */
+	void Construct(const FArguments& InArgs, const FText& InReportMessage, TArray<FStormSyncReportPackageData>& InPackageNames);
+
+	/** Closes the dialog. */
+	void CloseDialog();
+	
+private:
+	FStormSyncPackageReportNode PackageReportRootNode;
+	TSharedPtr<SStormSyncReportTreeView> ReportTreeView;
+
+	/** Brushes for the different node states */
+	const FSlateBrush* FolderOpenBrush = nullptr;
+	const FSlateBrush* FolderClosedBrush = nullptr;
+	const FSlateBrush* PackageBrush = nullptr;
+	
+	/** Recursively sets the checked/active state of every child of this node in the tree when a checkbox is toggled. */
+	static void SetStateRecursive(TSharedPtr<FStormSyncPackageReportNode> TreeItem, bool bIsChecked);
+
+	/** Callback to check whether a checkbox is checked or not. */
+	static ECheckBoxState GetEnabledCheckState(TSharedPtr<FStormSyncPackageReportNode> TreeItem);
+
+	/** Callback called whenever a checkbox is toggled. */
+	static void CheckBoxStateChanged(ECheckBoxState InCheckBoxState, TSharedPtr<FStormSyncPackageReportNode> TreeItem, TSharedRef<STableViewBase> OwnerTable);
+
+	/** Constructs the node tree given the list of package names */
+	void ConstructNodeTree(TArray<FStormSyncReportPackageData>& PackageNames);
+
+	/** Handler to generate a row in the report tree */
+	TSharedRef<ITableRow> GenerateTreeRow(TSharedPtr<FStormSyncPackageReportNode> TreeItem, const TSharedRef<STableViewBase>& OwnerTable) const;
+
+	/** Gets the children for the specified tree item */
+	static void GetChildrenForTree(TSharedPtr<FStormSyncPackageReportNode> TreeItem, TArray<TSharedPtr<FStormSyncPackageReportNode>>& OutChildren);
+	
+	/** Determines which image to display next to a node */
+	const FSlateBrush* GetNodeIcon(const TSharedPtr<FStormSyncPackageReportNode>& ReportNode) const;
+};
diff --git a/Engine/Plugins/Experimental/StormSync/Source/StormSyncEditor/Private/Slate/Status/SStormSyncStatusWidget.h b/Engine/Plugins/Experimental/StormSync/Source/StormSyncEditor/Private/Slate/Status/SStormSyncStatusWidget.h
--- a/Engine/Plugins/Experimental/StormSync/Source/StormSyncEditor/Private/Slate/Status/SStormSyncStatusWidget.h
+++ b/Engine/Plugins/Experimental/StormSync/Source/StormSyncEditor/Private/Slate/Status/SStormSyncStatusWidget.h
@@ -2,7 +2,6 @@
 
 #pragma once
 
-#include "CoreMinimal.h"
 #include "StormSyncPackageDescriptor.h"
 #include "Widgets/SCompoundWidget.h"
 #include "Widgets/Views/SListView.h"
@@ -12,7 +11,7 @@ struct FStormSyncImportFileInfo;
 struct FStormSyncTransportStatusResponse;
 
 /** A modal dialog to show the status of package names file state between two editor instances  */
-class STORMSYNCEDITOR_API SStormSyncStatusWidget : public SCompoundWidget
+class SStormSyncStatusWidget : public SCompoundWidget
 {
 public:
 	SLATE_BEGIN_ARGS(SStormSyncStatusWidget) {}
diff --git a/Engine/Plugins/Experimental/StormSync/Source/StormSyncEditor/Private/StormSyncEditor.cpp b/Engine/Plugins/Experimental/StormSync/Source/StormSyncEditor/Private/StormSyncEditor.cpp
--- a/Engine/Plugins/Experimental/StormSync/Source/StormSyncEditor/Private/StormSyncEditor.cpp
+++ b/Engine/Plugins/Experimental/StormSync/Source/StormSyncEditor/Private/StormSyncEditor.cpp
@@ -5,8 +5,11 @@
 #include "AssetRegistry/AssetData.h"
 #include "ContextMenu/StormSyncAssetFolderContextMenu.h"
 #include "Customization/StormSyncTransportSettingsDetailsCustomization.h"
+#include "Framework/Application/SlateApplication.h"
 #include "IStormSyncTransportClientModule.h"
+#include "Interfaces/IMainFrameModule.h"
 #include "PropertyEditorModule.h"
+#include "Slate/SStormSyncImportWizard.h"
 #include "Slate/Status/SStormSyncStatusWidget.h"
 #include "StormSyncCommandLineUtils.h"
 #include "StormSyncEditorLog.h"
@@ -56,6 +59,31 @@ void FStormSyncEditorModule::ShutdownModule()
 	UnregisterConsoleCommands();
 }
 
+TSharedRef<IStormSyncImportWizard> FStormSyncEditorModule::CreateWizard(const TArray<FStormSyncImportFileInfo>& InFilesToImport, const TArray<FStormSyncImportFileInfo>& InBufferFiles)
+{
+	const TSharedRef<SStormSyncImportWizard> ImportWizard = SNew(SStormSyncImportWizard, InFilesToImport, InBufferFiles);
+
+	const TSharedRef<SWindow> Window = SNew(SWindow)
+		.Title(LOCTEXT("Window_Title", "Storm Sync | Import files from local pak"))
+		.ClientSize(FVector2D(960, 700))
+		.SupportsMaximize(false)
+		.SupportsMinimize(false)
+		[
+			ImportWizard->AsShared()
+		];
+
+	TSharedPtr<SWindow> ParentWindow;
+
+	if (FModuleManager::Get().IsModuleLoaded("MainFrame"))
+	{
+		const IMainFrameModule& MainFrame = FModuleManager::LoadModuleChecked<IMainFrameModule>("MainFrame");
+		ParentWindow = MainFrame.GetParentWindow();
+	}
+
+	FSlateApplication::Get().AddModalWindow(Window, ParentWindow, false);
+	return ImportWizard;
+}
+
 void FStormSyncEditorModule::BuildPushAssetsMenuSection(FMenuBuilder& InMenuBuilder, const TArray<FName> InPackageNames, const bool bInIsPushing) const
 {
 	// Exposed to outside modules so that they're able to generate this submenu (useful for UI extensions)
diff --git a/Engine/Plugins/Experimental/StormSync/Source/StormSyncEditor/Private/StormSyncEditorLog.h b/Engine/Plugins/Experimental/StormSync/Source/StormSyncEditor/Private/StormSyncEditorLog.h
new file mode 100644
--- /dev/null
+++ b/Engine/Plugins/Experimental/StormSync/Source/StormSyncEditor/Private/StormSyncEditorLog.h
@@ -0,0 +1,12 @@
+// Copyright Epic Games, Inc. All Rights Reserved.
+
+#pragma once
+
+#include "Logging/LogMacros.h"
+
+DECLARE_LOG_CATEGORY_EXTERN(LogStormSyncEditor, Display, All);
+
+#define STORM_SYNC_EDITOR_LOG(Verbosity, Format, ...) \
+{ \
+	UE_LOG(LogStormSyncEditor, Verbosity, Format, ##__VA_ARGS__); \
+}
diff --git a/Engine/Plugins/Experimental/StormSync/Source/StormSyncEditor/Public/IStormSyncImportWizard.h b/Engine/Plugins/Experimental/StormSync/Source/StormSyncEditor/Public/IStormSyncImportWizard.h
new file mode 100644
--- /dev/null
+++ b/Engine/Plugins/Experimental/StormSync/Source/StormSyncEditor/Public/IStormSyncImportWizard.h
@@ -0,0 +1,10 @@
+// Copyright Epic Games, Inc. All Rights Reserved.
+
+#pragma once
+
+class IStormSyncImportWizard
+{
+public:
+	/** Returns user choice, either import or cancel */
+	virtual bool ShouldImport() const = 0;
+};
diff --git a/Engine/Plugins/Experimental/StormSync/Source/StormSyncEditor/Public/Slate/SStormSyncFileDependencyWidgetRow.h b/Engine/Plugins/Experimental/StormSync/Source/StormSyncEditor/Public/Slate/SStormSyncFileDependencyWidgetRow.h
deleted file mode 100644
diff --git a/Engine/Plugins/Experimental/StormSync/Source/StormSyncEditor/Public/Slate/SStormSyncImportWizard.h b/Engine/Plugins/Experimental/StormSync/Source/StormSyncEditor/Public/Slate/SStormSyncImportWizard.h
deleted file mode 100644
diff --git a/Engine/Plugins/Experimental/StormSync/Source/StormSyncEditor/Public/Slate/SStormSyncReportDialog.h b/Engine/Plugins/Experimental/StormSync/Source/StormSyncEditor/Public/Slate/SStormSyncReportDialog.h
deleted file mode 100644
diff --git a/Engine/Plugins/Experimental/StormSync/Source/StormSyncEditor/Public/StormSyncEditor.h b/Engine/Plugins/Experimental/StormSync/Source/StormSyncEditor/Public/StormSyncEditor.h
--- a/Engine/Plugins/Experimental/StormSync/Source/StormSyncEditor/Public/StormSyncEditor.h
+++ b/Engine/Plugins/Experimental/StormSync/Source/StormSyncEditor/Public/StormSyncEditor.h
@@ -2,7 +2,6 @@
 
 #pragma once
 
-#include "CoreMinimal.h"
 #include "IMessageContext.h"
 #include "Modules/ModuleInterface.h"
 #include "Modules/ModuleManager.h"
@@ -11,13 +10,14 @@
 class FExtender;
 class FMenuBuilder;
 class FStormSyncAssetFolderContextMenu;
+class IStormSyncImportWizard;
 struct FAssetData;
+struct FStormSyncImportFileInfo;
 
 /** Main entry point and implementation of StormSync Core Editor module. */
 class FStormSyncEditorModule : public IModuleInterface
 {
 public:
-	
 	/**
 	 * Singleton-like access to this module's interface.  This is just for convenience!
 	 * Beware of calling this during the shutdown phase, though.  Your module might have been unloaded already.
@@ -29,12 +29,15 @@ public:
 		static const FName ModuleName = "StormSyncEditor";
 		return FModuleManager::LoadModuleChecked<FStormSyncEditorModule>(ModuleName);
 	}
-	
+
 	//~ Begin IModuleInterface interface
 	virtual void StartupModule() override;
 	virtual void ShutdownModule() override;
 	//~ End IModuleInterface interface
-	
+
+	/** Opens the dialog in a modal new window */
+	STORMSYNCEDITOR_API TSharedRef<IStormSyncImportWizard> CreateWizard(const TArray<FStormSyncImportFileInfo>& InFilesToImport, const TArray<FStormSyncImportFileInfo>& InBufferFiles);
+
 	/**
 	 * Submenu extension delegate to build submenu for Push and Pull actions (useful for UI extensions)
 	 *
diff --git a/Engine/Plugins/Experimental/StormSync/Source/StormSyncEditor/Public/StormSyncEditorLog.h b/Engine/Plugins/Experimental/StormSync/Source/StormSyncEditor/Public/StormSyncEditorLog.h
deleted file mode 100644
diff --git a/Engine/Plugins/Experimental/StormSync/Source/StormSyncEditor/Public/Subsystems/StormSyncNotificationSubsystem.h b/Engine/Plugins/Experimental/StormSync/Source/StormSyncEditor/Public/Subsystems/StormSyncNotificationSubsystem.h
--- a/Engine/Plugins/Experimental/StormSync/Source/StormSyncEditor/Public/Subsystems/StormSyncNotificationSubsystem.h
+++ b/Engine/Plugins/Experimental/StormSync/Source/StormSyncEditor/Public/Subsystems/StormSyncNotificationSubsystem.h
@@ -2,7 +2,6 @@
 
 #pragma once
 
-#include "CoreMinimal.h"
 #include "EditorSubsystem.h"
 #include "Framework/Notifications/NotificationManager.h"
 #include "Logging/MessageLog.h"
@@ -30,34 +29,35 @@ struct FStormSyncNotificationTask
  * This subsystem main purpose is to listen for tcp related events triggered from StormSyncTransport layers and module,
  * to provide in-editor integration and handle UI/UX feedback to the end user.
  */
-UCLASS()
-class STORMSYNCEDITOR_API UStormSyncNotificationSubsystem : public UEditorSubsystem
+UCLASS(MinimalAPI)
+class UStormSyncNotificationSubsystem : public UEditorSubsystem
 {
 	GENERATED_BODY()
 
 public:
-	//~ Begin UEditorSubsystem interface
-	virtual void Initialize(FSubsystemCollectionBase& InCollection) override;
-	virtual void Deinitialize() override;
-	//~ End UEditorSubsystem interface
-
 	/** Static convenience method to return storm sync notification subsystem */
-	static UStormSyncNotificationSubsystem& Get();
+	STORMSYNCEDITOR_API static UStormSyncNotificationSubsystem& Get();
 
-	virtual TSharedPtr<SNotificationItem> AddSimpleNotification(const FText& InNotificationText);
-	virtual TSharedPtr<SNotificationItem> AddSimpleNotification(const FNotificationInfo& InInfo);
-	
-	virtual void NewPage(const FText& InLabel, const bool bInAppendDateTime = true);
-	virtual void Info(const FText& InMessageLogText);
-	virtual void Warning(const FText& InMessageLogText);
-	virtual void Error(const FText& InMessageLogText);
-	virtual void Notify(const EMessageSeverity::Type InSeverity, const FText& InMessageLogText, const bool bInForce = true);
+	TSharedPtr<SNotificationItem> AddSimpleNotification(const FText& InNotificationText);
+	TSharedPtr<SNotificationItem> AddSimpleNotification(const FNotificationInfo& InInfo);
+
+	void NewPage(const FText& InLabel, const bool bInAppendDateTime = true);
+	void Info(const FText& InMessageLogText);
+	void Warning(const FText& InMessageLogText);
+	void Error(const FText& InMessageLogText);
+	void Notify(const EMessageSeverity::Type InSeverity, const FText& InMessageLogText, const bool bInForce = true);
 
-	virtual void HandlePushResponse(const TSharedPtr<FStormSyncTransportPushResponse>& InResponse);
-	virtual void HandlePullResponse(const TSharedPtr<FStormSyncTransportPullResponse>& InResponse);
-	virtual void HandleSyncResponse(const TSharedPtr<FStormSyncTransportSyncResponse>& InResponse, const FText& InSuccessTextFormat, const FText& InErrorTextFormat);
+	STORMSYNCEDITOR_API void HandlePushResponse(const TSharedPtr<FStormSyncTransportPushResponse>& InResponse);
+	STORMSYNCEDITOR_API void HandlePullResponse(const TSharedPtr<FStormSyncTransportPullResponse>& InResponse);
+
+	void HandleSyncResponse(const TSharedPtr<FStormSyncTransportSyncResponse>& InResponse, const FText& InSuccessTextFormat, const FText& InErrorTextFormat);
 
 protected:
+	//~ Begin UEditorSubsystem interface
+	virtual void Initialize(FSubsystemCollectionBase& InCollection) override;
+	virtual void Deinitialize() override;
+	//~ End UEditorSubsystem interface
+
 	/** Currently active notifications. Key is the remote address. */
 	TMap<FString, FStormSyncNotificationTask> Notifications;
 
diff --git a/Engine/Plugins/Experimental/StormSync/Source/StormSyncImport/Private/StormSyncImport.cpp b/Engine/Plugins/Experimental/StormSync/Source/StormSyncImport/Private/StormSyncImport.cpp
--- a/Engine/Plugins/Experimental/StormSync/Source/StormSyncImport/Private/StormSyncImport.cpp
+++ b/Engine/Plugins/Experimental/StormSync/Source/StormSyncImport/Private/StormSyncImport.cpp
@@ -1,21 +1,8 @@
 // Copyright Epic Games, Inc. All Rights Reserved.
 
-#include "StormSyncImport.h"
-
 #include "StormSyncImportLog.h"
+#include "Modules/ModuleManager.h"
 
-#define LOCTEXT_NAMESPACE "FStormSyncImportModule"
-
-void FStormSyncImportModule::StartupModule()
-{
-	STORM_SYNC_IMPORT_LOG(Verbose, TEXT("Started StormSyncImport module ..."))
-}
-
-void FStormSyncImportModule::ShutdownModule()
-{
-    
-}
+DEFINE_LOG_CATEGORY(LogStormSyncImport);
 
-#undef LOCTEXT_NAMESPACE
-    
-IMPLEMENT_MODULE(FStormSyncImportModule, StormSyncImport)
\ No newline at end of file
+IMPLEMENT_MODULE(FDefaultModuleImpl, StormSyncImport)
diff --git a/Engine/Plugins/Experimental/StormSync/Source/StormSyncImport/Private/StormSyncImport.h b/Engine/Plugins/Experimental/StormSync/Source/StormSyncImport/Private/StormSyncImport.h
deleted file mode 100644
diff --git a/Engine/Plugins/Experimental/StormSync/Source/StormSyncImport/Private/StormSyncImportLog.cpp b/Engine/Plugins/Experimental/StormSync/Source/StormSyncImport/Private/StormSyncImportLog.cpp
deleted file mode 100644
diff --git a/Engine/Plugins/Experimental/StormSync/Source/StormSyncImport/Private/StormSyncImportLog.h b/Engine/Plugins/Experimental/StormSync/Source/StormSyncImport/Private/StormSyncImportLog.h
new file mode 100644
--- /dev/null
+++ b/Engine/Plugins/Experimental/StormSync/Source/StormSyncImport/Private/StormSyncImportLog.h
@@ -0,0 +1,12 @@
+// Copyright Epic Games, Inc. All Rights Reserved.
+
+#pragma once
+
+#include "Logging/LogMacros.h"
+
+STORMSYNCIMPORT_API DECLARE_LOG_CATEGORY_EXTERN(LogStormSyncImport, Display, All);
+
+#define STORM_SYNC_IMPORT_LOG(Verbosity, Format, ...) \
+{ \
+	UE_LOG(LogStormSyncImport, Verbosity, Format, ##__VA_ARGS__); \
+}
diff --git a/Engine/Plugins/Experimental/StormSync/Source/StormSyncImport/Private/Subsystems/StormSyncImportSubsystem.cpp b/Engine/Plugins/Experimental/StormSync/Source/StormSyncImport/Private/Subsystems/StormSyncImportSubsystem.cpp
--- a/Engine/Plugins/Experimental/StormSync/Source/StormSyncImport/Private/Subsystems/StormSyncImportSubsystem.cpp
+++ b/Engine/Plugins/Experimental/StormSync/Source/StormSyncImport/Private/Subsystems/StormSyncImportSubsystem.cpp
@@ -29,12 +29,13 @@
 #include "FileHelpers.h"
 #include "ISourceControlModule.h"
 #include "ISourceControlProvider.h"
+#include "IStormSyncImportWizard.h"
 #include "MessageLogModule.h"
 #include "ObjectTools.h"
 #include "PackageTools.h"
 #include "Settings/EditorLoadingSavingSettings.h"
-#include "Slate/SStormSyncImportWizard.h"
 #include "SourceControlHelpers.h"
+#include "StormSyncEditor.h"
 #include "Subsystems/AssetEditorSubsystem.h"
 #include "Toolkits/ToolkitManager.h"
 #endif
@@ -192,7 +193,7 @@ bool UStormSyncImportSubsystem::PerformImport(const FStormSyncPackageDescriptor&
 		if (bShowWizard)
 		{
 			// Create file import dialog. This is a modal dialog so it will return only after user selection
-			const TSharedRef<SStormSyncImportWizard> Wizard = SStormSyncImportWizard::CreateWizard(FilesToImport, BufferFiles);
+			const TSharedRef<IStormSyncImportWizard> Wizard = FStormSyncEditorModule::Get().CreateWizard(FilesToImport, BufferFiles);
 
 			// Early out if user canceled the operation
 			if (!Wizard->ShouldImport())
diff --git a/Engine/Plugins/Experimental/StormSync/Source/StormSyncImport/Private/Subsystems/StormSyncImportSubsystem.h b/Engine/Plugins/Experimental/StormSync/Source/StormSyncImport/Private/Subsystems/StormSyncImportSubsystem.h
new file mode 100644
--- /dev/null
+++ b/Engine/Plugins/Experimental/StormSync/Source/StormSyncImport/Private/Subsystems/StormSyncImportSubsystem.h
@@ -0,0 +1,175 @@
+// Copyright Epic Games, Inc. All Rights Reserved.
+
+#pragma once
+
+#include "Containers/Queue.h"
+#include "Containers/Ticker.h"
+#include "StormSyncCommonTypes.h"
+#include "StormSyncPackageDescriptor.h"
+#include "Subsystems/EngineSubsystem.h"
+#include "StormSyncImportSubsystem.generated.h"
+
+class IStormSyncImportSubsystemTask;
+class IToolkit;
+struct FStormSyncImportFileInfo;
+
+/** Holds information about extraction process for a given file */
+struct FStormSyncEditorFileReport
+{
+	/** Whether the extraction was successful or not */
+	bool bSuccess = true;
+	
+	/** Original package name */
+	FName PackageName;
+	
+	/** Fully qualified path we tried to extract to */
+	FString DestFilepath;
+
+	FStormSyncEditorFileReport() = default;
+};
+
+/**
+ * UStormSyncImportSubsystem
+ * 
+ * Subsystem for importing assets in the editor extracted from a local exported storm sync buffer, or coming from a network request.
+ *
+ * Contains utility functions and callbacks for hooking into importing.
+ */
+UCLASS()
+class STORMSYNCIMPORT_API UStormSyncImportSubsystem final : public UEngineSubsystem
+{
+	GENERATED_BODY()
+
+public:
+	//~ Begin USubsystem interface
+	virtual void Initialize(FSubsystemCollectionBase& Collection) override;
+	virtual void Deinitialize() override;
+	//~ End USubsystem interface
+
+	/** Static Convenience method to return storm sync import subsystem */
+	static UStormSyncImportSubsystem& Get();
+
+	/**
+	 * Enqueue an import task (and handle on next tick) only if no pending existing task already exist
+	 *
+	 * @param InImportTask Task to enqueue
+	 * @param InWorld World to get the timer manager from
+	 * @return true if the task was added, false otherwise. 
+	 */
+	bool EnqueueImportTask(const TSharedPtr<IStormSyncImportSubsystemTask>& InImportTask, const UWorld* InWorld);
+
+	/**
+	 * Takes an absolute filename (a storm sync pak) and begins extraction process
+	 *
+	 * Comparison between file coming from buffer and current file state in local file system is done to figure out
+	 * the list of files that needs to be extracted and updated in this editor instance.
+	 *
+	 * @param InFilename Absolute path name for the storm sync archive to import
+	 * 
+	 * @return Operation status, true if successful
+	 */
+	static bool PerformFileImport(const FString& InFilename);
+
+	/**
+	 * Takes a raw buffer (a storm sync pak) and begins extraction process
+	 *
+	 * Comparison between file coming from buffer and current file state in local file system is done to figure out
+	 * the list of files that needs to be extracted and updated in this editor instance.
+	 *
+	 * @param InPackageDescriptor The package descriptor for this buffer extraction. Only used for UI message logs.
+	 * @param InBuffer Raw buffer we'd like to extract files from to local project
+	 * 
+	 * @return Operation status, true if successful
+	 */
+	static bool PerformBufferImport(const FStormSyncPackageDescriptor& InPackageDescriptor, const FStormSyncBufferPtr& InBuffer);
+
+	/**
+	 * Takes a raw buffer (a storm sync pak) and begins extraction process
+	 *
+	 * Comparison between file coming from buffer and current file state in local file system is done to figure out
+	 * the list of files that needs to be extracted and updated in this editor instance.
+	 *
+	 * @param InPackageDescriptor The package descriptor for this buffer extraction. Only used for UI message logs.
+	 * @param InBuffer Raw buffer we'd like to extract files from to local project
+	 * @param bShowWizard Whether or not to show the import wizard
+	 * @param bDryRun Whether we should proceed to extraction or simply log to console buffer content
+	 * 
+	 * @return Operation status, true if successful
+	 */
+	static bool PerformImport(const FStormSyncPackageDescriptor& InPackageDescriptor, const FStormSyncBufferPtr& InBuffer, bool bShowWizard = false, bool bDryRun = false);
+
+private:
+	/** The name of the message output log we will send messages to */
+	static constexpr const TCHAR* LogName = TEXT("StormSyncImport");
+
+	FTSTicker::FDelegateHandle NextTickHandler;
+
+	/* Tasks waiting to be run next tick */
+	TQueue<TSharedPtr<IStormSyncImportSubsystemTask>> PendingTasks;
+	
+	/** Store closed assets to reopen during extraction process */
+	TArray<FString> ClosedPackageNames;
+	
+	/** Store extraction process data */
+	TArray<FStormSyncEditorFileReport> ExtractedFileReports;
+
+	/* Run deferred logic waiting to be run next tick */
+	bool HandleNextTick(float InDeltaTime);
+	
+	/**
+	 * Checks local project file for diff against provided file info and dest path.
+	 *
+	 * If a change is detected (either missing locally or mismatched file size / hash), OutFilesToImport is filled with a new file info indicating that the
+	 * file needs an import.
+	 */
+	bool FillFilesToImport(const FStormSyncImportFileInfo& InFileInfo, const FString& InDestFilepath, TArray<FStormSyncImportFileInfo>& OutFilesToImport) const;
+
+	/** Event handler for core delegate when a pak extraction starts */
+	void HandlePakPreExtract(const FStormSyncPackageDescriptor& InPackageDescriptor, int32 FileCount);
+
+	/** Event handler for core delegate when a pak extraction is ended */
+	void HandlePakPostExtract(const FStormSyncPackageDescriptor& InPackageDescriptor, int32 FileCount) const;
+
+	/** Event handler for existing asset checkout prior to extracting the new overwriting assets */
+	void HandleExistingAssets(TConstArrayView<const FStormSyncImportFileInfo*> InExistingFiles, bool bInShowPrompt);
+
+	/** Event handler for auto-adding new assets post extracting them */
+	bool HandleNewAssets(TConstArrayView<const FStormSyncImportFileInfo*> InNewFiles, bool bInShowPrompt);
+
+	/** Event handler for core delegate when a file is extracted from an incoming pak */
+	void HandlePakAssetExtract(const FStormSyncFileDependency& FileDependency, const FString& DestFilepath, const FStormSyncBufferPtr& FileBuffer);
+
+	/** Checks if given asset is currently being edited, and if so, closes the editor (required prior to file deletion) */
+	static void CloseEditors(const TArray<FAssetData>& InAssets, TArray<FAssetData>& OutClosedAssets);
+	
+	/** Takes a list of package names to ask the Asset Editor subsystem to open */
+	static void OpenClosedEditors(const TArray<FString>& ClosedPackageNames);
+
+	/**
+	 * Our own custom version of ObjectTools::DeleteAssets
+	 *
+	 * We want to force delete the file, to allow assets with on disk references (for things like Textures in an Avalanche Blueprints) to be properly overwritten,
+	 * if referenced in memory, opened in editor etc.
+	 *
+	 * ObjectTools::DeleteAssets will only succeed if the AssetDeleteModel doesn't have on disk references (see CanDelete / CanForceDelete)
+	 */
+	static bool DeleteAssets(const TArray<FAssetData>& AssetsToDelete, bool bShowConfirmation = false);
+
+	/** Makes the given packages writeable. Should only be used if Source Control is disabled. Returns number of packages that succeeded */
+	static int32 MakePackagesWriteable(TConstArrayView<UPackage*> InPackages);
+
+	/** Checks if given UObject asset is currently being edited by an Asset Editor (is it opened in editor right now?) */
+	static bool IsAssetCurrentlyBeingEdited(const TSharedPtr<IToolkit>& InAssetEditor, const UObject* InAsset);
+
+	/** Write file buffer to destination filepath */
+	static bool WriteFile(const FString& DestFilepath, const uint64& FileSize, uint8* FileBuffer);
+
+	/** Blocks till all pending package/ linker requests are fulfilled. */
+	static void FlushPackageLoading(const FString& InPackageName, bool bInForceBulkDataLoad = true);
+
+	/** Attempts to reload the specified top-level packages. */
+	static void HotReloadPackages(const TArray<FStormSyncEditorFileReport>& InExtractedFileReports, bool bInInteractiveHotReload);
+
+	/** Generates two arrays with pointers to the input file array. One with the already-existing files, and the other with the new files that would be generated in import */
+	static void SplitImportFiles(TConstArrayView<FStormSyncImportFileInfo> InFiles, TArray<const FStormSyncImportFileInfo*>& OutExistingFiles, TArray<const FStormSyncImportFileInfo*>& OutNewFiles);
+};
diff --git a/Engine/Plugins/Experimental/StormSync/Source/StormSyncImport/Private/Subsystems/StormSyncImportWorldSubsystem.h b/Engine/Plugins/Experimental/StormSync/Source/StormSyncImport/Private/Subsystems/StormSyncImportWorldSubsystem.h
new file mode 100644
--- /dev/null
+++ b/Engine/Plugins/Experimental/StormSync/Source/StormSyncImport/Private/Subsystems/StormSyncImportWorldSubsystem.h
@@ -0,0 +1,30 @@
+// Copyright Epic Games, Inc. All Rights Reserved.
+
+#pragma once
+
+#include "StormSyncCommonTypes.h"
+#include "StormSyncPackageDescriptor.h"
+#include "Subsystems/WorldSubsystem.h"
+#include "StormSyncImportWorldSubsystem.generated.h"
+
+/**
+ * UStormSyncImportWorldSubsystem
+ * 
+ * **World** Subsystem for importing assets in -game mode, from a local exported storm sync buffer, or coming from a network request.
+ *
+ * Bulk of the import logic is handled by UStormSyncImportSubsystem, this world subsystem registers a delegate on FStormSyncCoreDelegates
+ * OnRequestImportFile / 
+ */
+
+UCLASS()
+class STORMSYNCIMPORT_API UStormSyncImportWorldSubsystem : public UWorldSubsystem
+{
+	GENERATED_BODY()
+
+public:
+	virtual void Initialize(FSubsystemCollectionBase& Collection) override;
+	virtual void Deinitialize() override;
+	
+	void HandleImportFile(const FString& InFilename) const;
+	void HandleImportBuffer(const FStormSyncPackageDescriptor& InPackageDescriptor, const FStormSyncBufferPtr& InBuffer) const;
+};
diff --git a/Engine/Plugins/Experimental/StormSync/Source/StormSyncImport/Private/Tasks/IStormSyncImportTask.h b/Engine/Plugins/Experimental/StormSync/Source/StormSyncImport/Private/Tasks/IStormSyncImportTask.h
new file mode 100644
--- /dev/null
+++ b/Engine/Plugins/Experimental/StormSync/Source/StormSyncImport/Private/Tasks/IStormSyncImportTask.h
@@ -0,0 +1,12 @@
+// Copyright Epic Games, Inc. All Rights Reserved.
+
+#pragma once
+
+/** Interface for tasks that need delayed execution */
+class IStormSyncImportSubsystemTask
+{
+public:
+	virtual ~IStormSyncImportSubsystemTask() = default;
+
+	virtual void Run() = 0;
+};
diff --git a/Engine/Plugins/Experimental/StormSync/Source/StormSyncImport/Private/Tasks/StormSyncImportBufferTask.h b/Engine/Plugins/Experimental/StormSync/Source/StormSyncImport/Private/Tasks/StormSyncImportBufferTask.h
--- a/Engine/Plugins/Experimental/StormSync/Source/StormSyncImport/Private/Tasks/StormSyncImportBufferTask.h
+++ b/Engine/Plugins/Experimental/StormSync/Source/StormSyncImport/Private/Tasks/StormSyncImportBufferTask.h
@@ -2,7 +2,6 @@
 
 #pragma once
 
-#include "CoreMinimal.h"
 #include "StormSyncPackageDescriptor.h"
 #include "Tasks/IStormSyncImportTask.h"
 
diff --git a/Engine/Plugins/Experimental/StormSync/Source/StormSyncImport/Private/Tasks/StormSyncImportFilesTask.h b/Engine/Plugins/Experimental/StormSync/Source/StormSyncImport/Private/Tasks/StormSyncImportFilesTask.h
--- a/Engine/Plugins/Experimental/StormSync/Source/StormSyncImport/Private/Tasks/StormSyncImportFilesTask.h
+++ b/Engine/Plugins/Experimental/StormSync/Source/StormSyncImport/Private/Tasks/StormSyncImportFilesTask.h
@@ -2,7 +2,7 @@
 
 #pragma once
 
-#include "CoreMinimal.h"
+#include "Containers/UnrealString.h"
 #include "Tasks/IStormSyncImportTask.h"
 
 /** Import files from local exported buffer implementation for tasks that need delayed execution */
diff --git a/Engine/Plugins/Experimental/StormSync/Source/StormSyncImport/Public/IStormSyncImportModule.h b/Engine/Plugins/Experimental/StormSync/Source/StormSyncImport/Public/IStormSyncImportModule.h
deleted file mode 100644
diff --git a/Engine/Plugins/Experimental/StormSync/Source/StormSyncImport/Public/StormSyncImportLog.h b/Engine/Plugins/Experimental/StormSync/Source/StormSyncImport/Public/StormSyncImportLog.h
deleted file mode 100644
diff --git a/Engine/Plugins/Experimental/StormSync/Source/StormSyncImport/Public/Subsystems/StormSyncImportSubsystem.h b/Engine/Plugins/Experimental/StormSync/Source/StormSyncImport/Public/Subsystems/StormSyncImportSubsystem.h
deleted file mode 100644
diff --git a/Engine/Plugins/Experimental/StormSync/Source/StormSyncImport/Public/Subsystems/StormSyncImportWorldSubsystem.h b/Engine/Plugins/Experimental/StormSync/Source/StormSyncImport/Public/Subsystems/StormSyncImportWorldSubsystem.h
deleted file mode 100644
diff --git a/Engine/Plugins/Experimental/StormSync/Source/StormSyncImport/Public/Tasks/IStormSyncImportTask.h b/Engine/Plugins/Experimental/StormSync/Source/StormSyncImport/Public/Tasks/IStormSyncImportTask.h
deleted file mode 100644
diff --git a/Engine/Plugins/Experimental/StormSync/Source/StormSyncImport/StormSyncImport.Build.cs b/Engine/Plugins/Experimental/StormSync/Source/StormSyncImport/StormSyncImport.Build.cs
--- a/Engine/Plugins/Experimental/StormSync/Source/StormSyncImport/StormSyncImport.Build.cs
+++ b/Engine/Plugins/Experimental/StormSync/Source/StormSyncImport/StormSyncImport.Build.cs
@@ -7,16 +7,10 @@ public class StormSyncImport : ModuleRules
 	{
 		PCHUsage = ModuleRules.PCHUsageMode.UseExplicitOrSharedPCHs;
 
-		PublicDependencyModuleNames.AddRange(
-			new string[]
-			{
-				"Core",
-			}
-		);
-
 		PrivateDependencyModuleNames.AddRange(
 			new string[]
 			{
+				"Core",
 				"CoreUObject",
 				"Engine",
 				"MessageLog",
diff --git a/Engine/Plugins/Experimental/StormSync/Source/StormSyncTests/Private/StormSyncTests.cpp b/Engine/Plugins/Experimental/StormSync/Source/StormSyncTests/Private/StormSyncTests.cpp
--- a/Engine/Plugins/Experimental/StormSync/Source/StormSyncTests/Private/StormSyncTests.cpp
+++ b/Engine/Plugins/Experimental/StormSync/Source/StormSyncTests/Private/StormSyncTests.cpp
@@ -1,16 +1,5 @@
 // Copyright Epic Games, Inc. All Rights Reserved.
 
-#include "StormSyncTests.h"
-
 #include "Modules/ModuleManager.h"
 
-void FStormSyncTestsModule::StartupModule()
-{
-}
-
-void FStormSyncTestsModule::ShutdownModule()
-{
-
-}
-
-IMPLEMENT_MODULE(FStormSyncTestsModule, StormSyncTests)
+IMPLEMENT_MODULE(FDefaultModuleImpl, StormSyncTests)
diff --git a/Engine/Plugins/Experimental/StormSync/Source/StormSyncTests/Public/StormSyncTests.h b/Engine/Plugins/Experimental/StormSync/Source/StormSyncTests/Public/StormSyncTests.h
deleted file mode 100644
diff --git a/Engine/Plugins/Experimental/StormSync/Source/StormSyncTests/StormSyncTests.Build.cs b/Engine/Plugins/Experimental/StormSync/Source/StormSyncTests/StormSyncTests.Build.cs
--- a/Engine/Plugins/Experimental/StormSync/Source/StormSyncTests/StormSyncTests.Build.cs
+++ b/Engine/Plugins/Experimental/StormSync/Source/StormSyncTests/StormSyncTests.Build.cs
@@ -6,18 +6,12 @@ public class StormSyncTests : ModuleRules
 {
 	public StormSyncTests(ReadOnlyTargetRules Target) : base(Target)
 	{
-		PublicDependencyModuleNames.AddRange(
+		PrivateDependencyModuleNames.AddRange(
 			new string[]
 			{
 				"Core",
 				"CoreUObject",
-				"Engine"
-			}
-		);
-
-		PrivateDependencyModuleNames.AddRange(
-			new string[]
-			{
+				"Engine",
 				"Projects",
 				"StormSyncCore",
 				"StormSyncDrives",
diff --git a/Engine/Plugins/Experimental/StormSync/Source/StormSyncTransportClient/Private/Services/StormSyncPullMessageService.h b/Engine/Plugins/Experimental/StormSync/Source/StormSyncTransportClient/Private/Services/StormSyncPullMessageService.h
--- a/Engine/Plugins/Experimental/StormSync/Source/StormSyncTransportClient/Private/Services/StormSyncPullMessageService.h
+++ b/Engine/Plugins/Experimental/StormSync/Source/StormSyncTransportClient/Private/Services/StormSyncPullMessageService.h
@@ -2,7 +2,6 @@
 
 #pragma once
 
-#include "CoreMinimal.h"
 #include "IStormSyncTransportClientLocalEndpoint.h"
 #include "IStormSyncTransportMessageService.h"
 #include "MessageEndpoint.h"
diff --git a/Engine/Plugins/Experimental/StormSync/Source/StormSyncTransportClient/Private/Services/StormSyncPushMessageService.h b/Engine/Plugins/Experimental/StormSync/Source/StormSyncTransportClient/Private/Services/StormSyncPushMessageService.h
--- a/Engine/Plugins/Experimental/StormSync/Source/StormSyncTransportClient/Private/Services/StormSyncPushMessageService.h
+++ b/Engine/Plugins/Experimental/StormSync/Source/StormSyncTransportClient/Private/Services/StormSyncPushMessageService.h
@@ -2,7 +2,6 @@
 
 #pragma once
 
-#include "CoreMinimal.h"
 #include "IStormSyncTransportClientLocalEndpoint.h"
 #include "IStormSyncTransportMessageService.h"
 
@@ -15,7 +14,7 @@ class FStormSyncPushMessageService : public IStormSyncTransportMessageService
 {
 public:
 	explicit FStormSyncPushMessageService(TSharedPtr<IStormSyncTransportLocalEndpoint, ESPMode::ThreadSafe> InLocalEndpoint);
-	virtual ~FStormSyncPushMessageService();
+	virtual ~FStormSyncPushMessageService() override;
 
 	//~ Begin IStormSyncTransportMessageService interface
 	virtual void InitializeMessageEndpoint(FMessageEndpointBuilder& InEndpointBuilder) override;
diff --git a/Engine/Plugins/Experimental/StormSync/Source/StormSyncTransportClient/Private/Services/StormSyncStatusMessageService.h b/Engine/Plugins/Experimental/StormSync/Source/StormSyncTransportClient/Private/Services/StormSyncStatusMessageService.h
--- a/Engine/Plugins/Experimental/StormSync/Source/StormSyncTransportClient/Private/Services/StormSyncStatusMessageService.h
+++ b/Engine/Plugins/Experimental/StormSync/Source/StormSyncTransportClient/Private/Services/StormSyncStatusMessageService.h
@@ -2,7 +2,6 @@
 
 #pragma once
 
-#include "CoreMinimal.h"
 #include "IStormSyncTransportClientLocalEndpoint.h"
 #include "IStormSyncTransportMessageService.h"
 
@@ -14,7 +13,7 @@ class FStormSyncStatusMessageService : public IStormSyncTransportMessageService
 {
 public:
 	explicit FStormSyncStatusMessageService(TSharedPtr<IStormSyncTransportLocalEndpoint, ESPMode::ThreadSafe> InLocalEndpoint);
-	virtual ~FStormSyncStatusMessageService();
+	virtual ~FStormSyncStatusMessageService() override;
 
 	//~ Begin IStormSyncTransportMessageService interface
 	virtual void InitializeMessageEndpoint(FMessageEndpointBuilder& InEndpointBuilder) override;
diff --git a/Engine/Plugins/Experimental/StormSync/Source/StormSyncTransportClient/Private/Socket/StormSyncTransportClientSocket.h b/Engine/Plugins/Experimental/StormSync/Source/StormSyncTransportClient/Private/Socket/StormSyncTransportClientSocket.h
--- a/Engine/Plugins/Experimental/StormSync/Source/StormSyncTransportClient/Private/Socket/StormSyncTransportClientSocket.h
+++ b/Engine/Plugins/Experimental/StormSync/Source/StormSyncTransportClient/Private/Socket/StormSyncTransportClientSocket.h
@@ -2,7 +2,6 @@
 
 #pragma once
 
-#include "CoreMinimal.h"
 #include "HAL/Runnable.h"
 #include "Interfaces/IPv4/IPv4Endpoint.h"
 #include "StormSyncTransportSettings.h"
diff --git a/Engine/Plugins/Experimental/StormSync/Source/StormSyncTransportClient/Private/StormSyncTransportClientEndpoint.h b/Engine/Plugins/Experimental/StormSync/Source/StormSyncTransportClient/Private/StormSyncTransportClientEndpoint.h
--- a/Engine/Plugins/Experimental/StormSync/Source/StormSyncTransportClient/Private/StormSyncTransportClientEndpoint.h
+++ b/Engine/Plugins/Experimental/StormSync/Source/StormSyncTransportClient/Private/StormSyncTransportClientEndpoint.h
@@ -2,7 +2,6 @@
 
 #pragma once
 
-#include "CoreMinimal.h"
 #include "IStormSyncTransportClientLocalEndpoint.h"
 #include "StormSyncTransportMessages.h"
 #include "Interfaces/IPv4/IPv4Endpoint.h"
@@ -43,7 +42,7 @@ class FStormSyncTransportClientEndpoint : public IStormSyncTransportClientLocalE
 {
 public:
 	FStormSyncTransportClientEndpoint();
-	virtual ~FStormSyncTransportClientEndpoint();
+	virtual ~FStormSyncTransportClientEndpoint() override;
 
 	/** Initializes the message endpoint. */
 	void InitializeMessaging(const FString& InEndpointFriendlyName = TEXT("Client"));
diff --git a/Engine/Plugins/Experimental/StormSync/Source/StormSyncTransportClient/Private/StormSyncTransportClientLog.h b/Engine/Plugins/Experimental/StormSync/Source/StormSyncTransportClient/Private/StormSyncTransportClientLog.h
new file mode 100644
--- /dev/null
+++ b/Engine/Plugins/Experimental/StormSync/Source/StormSyncTransportClient/Private/StormSyncTransportClientLog.h
@@ -0,0 +1,12 @@
+// Copyright Epic Games, Inc. All Rights Reserved.
+
+#pragma once
+
+#include "Logging/LogMacros.h"
+
+DECLARE_LOG_CATEGORY_EXTERN(LogStormSyncClient, Display, All);
+
+#define STORM_SYNC_CLIENT_LOG(Verbosity, Format, ...) \
+{ \
+    UE_LOG(LogStormSyncClient, Verbosity, Format, ##__VA_ARGS__); \
+}
diff --git a/Engine/Plugins/Experimental/StormSync/Source/StormSyncTransportClient/Private/StormSyncTransportClientModule.cpp b/Engine/Plugins/Experimental/StormSync/Source/StormSyncTransportClient/Private/StormSyncTransportClientModule.cpp
--- a/Engine/Plugins/Experimental/StormSync/Source/StormSyncTransportClient/Private/StormSyncTransportClientModule.cpp
+++ b/Engine/Plugins/Experimental/StormSync/Source/StormSyncTransportClient/Private/StormSyncTransportClientModule.cpp
@@ -4,7 +4,6 @@
 
 #include "StormSyncCommandLineUtils.h"
 #include "StormSyncCoreUtils.h"
-#include "IStormSyncTransportCoreModule.h"
 #include "IStormSyncTransportLocalEndpoint.h"
 #include "MessageEndpoint.h"
 #include "Misc/CoreDelegates.h"
diff --git a/Engine/Plugins/Experimental/StormSync/Source/StormSyncTransportClient/Public/IStormSyncTransportClientLocalEndpoint.h b/Engine/Plugins/Experimental/StormSync/Source/StormSyncTransportClient/Public/IStormSyncTransportClientLocalEndpoint.h
--- a/Engine/Plugins/Experimental/StormSync/Source/StormSyncTransportClient/Public/IStormSyncTransportClientLocalEndpoint.h
+++ b/Engine/Plugins/Experimental/StormSync/Source/StormSyncTransportClient/Public/IStormSyncTransportClientLocalEndpoint.h
@@ -2,7 +2,6 @@
 
 #pragma once
 
-#include "CoreMinimal.h"
 #include "IStormSyncTransportLocalEndpoint.h"
 #include "StormSyncTransportMessages.h"
 
@@ -12,9 +11,8 @@ DECLARE_DELEGATE_OneParam(FOnStormSyncRequestStatusComplete, const TSharedPtr<FS
 DECLARE_DELEGATE_OneParam(FOnStormSyncPushComplete, const TSharedPtr<FStormSyncTransportPushResponse>&);
 DECLARE_DELEGATE_OneParam(FOnStormSyncPullComplete, const TSharedPtr<FStormSyncTransportPullResponse>&);
 
-
 /** Local Endpoint interface for Storm Sync Client */
-class STORMSYNCTRANSPORTCLIENT_API IStormSyncTransportClientLocalEndpoint : public IStormSyncTransportLocalEndpoint
+class IStormSyncTransportClientLocalEndpoint : public IStormSyncTransportLocalEndpoint
 {
 public:
 	/** Request a status from the given remote address. The given delegate will be called when the response comes in */
diff --git a/Engine/Plugins/Experimental/StormSync/Source/StormSyncTransportClient/Public/IStormSyncTransportClientModule.h b/Engine/Plugins/Experimental/StormSync/Source/StormSyncTransportClient/Public/IStormSyncTransportClientModule.h
--- a/Engine/Plugins/Experimental/StormSync/Source/StormSyncTransportClient/Public/IStormSyncTransportClientModule.h
+++ b/Engine/Plugins/Experimental/StormSync/Source/StormSyncTransportClient/Public/IStormSyncTransportClientModule.h
@@ -2,7 +2,6 @@
 
 #pragma once
 
-#include "CoreMinimal.h"
 #include "IStormSyncTransportClientLocalEndpoint.h"
 #include "StormSyncTransportMessages.h"
 #include "Modules/ModuleInterface.h"
@@ -13,7 +12,7 @@ struct FMessageAddress;
 
 using FMessageEndpointSharedPtr = TSharedPtr<FMessageEndpoint, ESPMode::ThreadSafe>;
 
-class STORMSYNCTRANSPORTCLIENT_API IStormSyncTransportClientModule : public IModuleInterface
+class IStormSyncTransportClientModule : public IModuleInterface
 {
 public:
 	/**
diff --git a/Engine/Plugins/Experimental/StormSync/Source/StormSyncTransportClient/Public/StormSyncTransportClientLog.h b/Engine/Plugins/Experimental/StormSync/Source/StormSyncTransportClient/Public/StormSyncTransportClientLog.h
deleted file mode 100644
diff --git a/Engine/Plugins/Experimental/StormSync/Source/StormSyncTransportCore/Private/StormSyncTransportCoreLog.h b/Engine/Plugins/Experimental/StormSync/Source/StormSyncTransportCore/Private/StormSyncTransportCoreLog.h
new file mode 100644
--- /dev/null
+++ b/Engine/Plugins/Experimental/StormSync/Source/StormSyncTransportCore/Private/StormSyncTransportCoreLog.h
@@ -0,0 +1,12 @@
+// Copyright Epic Games, Inc. All Rights Reserved.
+
+#pragma once
+
+#include "Logging/LogMacros.h"
+
+DECLARE_LOG_CATEGORY_EXTERN(LogStormSyncTransportCore, Display, All);
+
+#define STORM_SYNC_CORE_LOG(Verbosity, Format, ...) \
+{ \
+    UE_LOG(LogStormSyncTransportCore, Verbosity, Format, ##__VA_ARGS__); \
+}
diff --git a/Engine/Plugins/Experimental/StormSync/Source/StormSyncTransportCore/Private/StormSyncTransportCoreModule.cpp b/Engine/Plugins/Experimental/StormSync/Source/StormSyncTransportCore/Private/StormSyncTransportCoreModule.cpp
--- a/Engine/Plugins/Experimental/StormSync/Source/StormSyncTransportCore/Private/StormSyncTransportCoreModule.cpp
+++ b/Engine/Plugins/Experimental/StormSync/Source/StormSyncTransportCore/Private/StormSyncTransportCoreModule.cpp
@@ -1,18 +1,5 @@
 // Copyright Epic Games, Inc. All Rights Reserved.
 
-#include "StormSyncTransportCoreModule.h"
+#include "Modules/ModuleManager.h"
 
-#define LOCTEXT_NAMESPACE "StormSyncTransportCoreModule"
-
-
-void FStormSyncTransportCoreModule::StartupModule()
-{
-}
-
-void FStormSyncTransportCoreModule::ShutdownModule()
-{
-}
-
-IMPLEMENT_MODULE(FStormSyncTransportCoreModule, StormSyncTransportCore)
-
-#undef LOCTEXT_NAMESPACE
+IMPLEMENT_MODULE(FDefaultModuleImpl, StormSyncTransportCore)
diff --git a/Engine/Plugins/Experimental/StormSync/Source/StormSyncTransportCore/Private/StormSyncTransportCoreModule.h b/Engine/Plugins/Experimental/StormSync/Source/StormSyncTransportCore/Private/StormSyncTransportCoreModule.h
deleted file mode 100644
diff --git a/Engine/Plugins/Experimental/StormSync/Source/StormSyncTransportCore/Private/StormSyncTransportSettings.cpp b/Engine/Plugins/Experimental/StormSync/Source/StormSyncTransportCore/Private/StormSyncTransportSettings.cpp
--- a/Engine/Plugins/Experimental/StormSync/Source/StormSyncTransportCore/Private/StormSyncTransportSettings.cpp
+++ b/Engine/Plugins/Experimental/StormSync/Source/StormSyncTransportCore/Private/StormSyncTransportSettings.cpp
@@ -26,17 +26,7 @@ const UStormSyncTransportSettings& UStormSyncTransportSettings::Get()
 	return *Settings;
 }
 
-FName UStormSyncTransportSettings::GetCategoryName() const
-{
-	return PluginCategoryName;
-}
-
 #if WITH_EDITOR
-FText UStormSyncTransportSettings::GetSectionText() const
-{
-	return NSLOCTEXT("StormSyncTransportSettings", "SettingsSection", "Transport & Network");
-}
-
 void UStormSyncTransportSettings::OpenEditorSettingsWindow() const
 {
 	static ISettingsModule& SettingsModule = FModuleManager::LoadModuleChecked<ISettingsModule>("Settings");
diff --git a/Engine/Plugins/Experimental/StormSync/Source/StormSyncTransportCore/Public/IStormSyncTransportCoreModule.h b/Engine/Plugins/Experimental/StormSync/Source/StormSyncTransportCore/Public/IStormSyncTransportCoreModule.h
deleted file mode 100644
diff --git a/Engine/Plugins/Experimental/StormSync/Source/StormSyncTransportCore/Public/IStormSyncTransportLocalEndpoint.h b/Engine/Plugins/Experimental/StormSync/Source/StormSyncTransportCore/Public/IStormSyncTransportLocalEndpoint.h
--- a/Engine/Plugins/Experimental/StormSync/Source/StormSyncTransportCore/Public/IStormSyncTransportLocalEndpoint.h
+++ b/Engine/Plugins/Experimental/StormSync/Source/StormSyncTransportCore/Public/IStormSyncTransportLocalEndpoint.h
@@ -2,7 +2,7 @@
 
 #pragma once
 
-#include "CoreMinimal.h"
+#include "Templates/SharedPointer.h"
 
 class FMessageEndpoint;
 
@@ -13,9 +13,7 @@ class IStormSyncTransportLocalEndpoint : public TSharedFromThis<IStormSyncTransp
 {
 public:
 	/** Virtual destructor */
-	virtual ~IStormSyncTransportLocalEndpoint()
-	{
-	}
+	virtual ~IStormSyncTransportLocalEndpoint() = default;
 
 	/** Returns whether message endpoint is currently active */
 	virtual bool IsRunning() const = 0;
diff --git a/Engine/Plugins/Experimental/StormSync/Source/StormSyncTransportCore/Public/IStormSyncTransportMessageService.h b/Engine/Plugins/Experimental/StormSync/Source/StormSyncTransportCore/Public/IStormSyncTransportMessageService.h
--- a/Engine/Plugins/Experimental/StormSync/Source/StormSyncTransportCore/Public/IStormSyncTransportMessageService.h
+++ b/Engine/Plugins/Experimental/StormSync/Source/StormSyncTransportCore/Public/IStormSyncTransportMessageService.h
@@ -2,7 +2,7 @@
 
 #pragma once
 
-#include "CoreMinimal.h"
+#include "Templates/SharedPointer.h"
 
 struct FMessageEndpointBuilder;
 
diff --git a/Engine/Plugins/Experimental/StormSync/Source/StormSyncTransportCore/Public/StormSyncTransportCoreLog.h b/Engine/Plugins/Experimental/StormSync/Source/StormSyncTransportCore/Public/StormSyncTransportCoreLog.h
deleted file mode 100644
diff --git a/Engine/Plugins/Experimental/StormSync/Source/StormSyncTransportCore/Public/StormSyncTransportMessages.h b/Engine/Plugins/Experimental/StormSync/Source/StormSyncTransportCore/Public/StormSyncTransportMessages.h
--- a/Engine/Plugins/Experimental/StormSync/Source/StormSyncTransportCore/Public/StormSyncTransportMessages.h
+++ b/Engine/Plugins/Experimental/StormSync/Source/StormSyncTransportCore/Public/StormSyncTransportMessages.h
@@ -23,7 +23,7 @@ struct FStormSyncTransportHeartbeatMessage
 
 /** Holds data about local project (such as Message Address Ids, Project Name, Hostname, InstanceId, etc.) */
 USTRUCT()
-struct STORMSYNCTRANSPORTCORE_API FStormSyncConnectionInfo
+struct FStormSyncConnectionInfo
 {
 	GENERATED_BODY()
 
@@ -64,18 +64,18 @@ struct STORMSYNCTRANSPORTCORE_API FStormSyncConnectionInfo
 	FString ProjectDir;
 
 	/** Default constructor */
-	FStormSyncConnectionInfo();
+	STORMSYNCTRANSPORTCORE_API FStormSyncConnectionInfo();
 
 	/** Returns debug string for this message */
-	FString ToString() const;
+	STORMSYNCTRANSPORTCORE_API FString ToString() const;
 
 	/** Internal helper to return the last portion of a path, similar to the Unix basename command. Trailing directory separators are ignored */
-	static FString GetBasename(const FString& InPath);
+	STORMSYNCTRANSPORTCORE_API static FString GetBasename(const FString& InPath);
 };
 
 /** Connection message indicating a Connect event with data about local project / application */
 USTRUCT()
-struct STORMSYNCTRANSPORTCORE_API FStormSyncTransportConnectMessage : public FStormSyncConnectionInfo
+struct FStormSyncTransportConnectMessage : public FStormSyncConnectionInfo
 {
 	GENERATED_BODY()
 };
@@ -300,7 +300,7 @@ struct FStormSyncTransportPullResponse : public FStormSyncTransportSyncResponse
 };
 
 USTRUCT()
-struct STORMSYNCTRANSPORTCORE_API FStormSyncTransportStatusRequest
+struct FStormSyncTransportStatusRequest
 {
 	GENERATED_BODY()
 
@@ -332,7 +332,7 @@ struct STORMSYNCTRANSPORTCORE_API FStormSyncTransportStatusRequest
 };
 
 USTRUCT()
-struct STORMSYNCTRANSPORTCORE_API FStormSyncTransportStatusResponse
+struct FStormSyncTransportStatusResponse
 {
 	GENERATED_BODY()
 	
@@ -366,13 +366,13 @@ struct STORMSYNCTRANSPORTCORE_API FStormSyncTransportStatusResponse
 };
 
 USTRUCT()
-struct STORMSYNCTRANSPORTCORE_API FStormSyncTransportStatusPing
+struct FStormSyncTransportStatusPing
 {
 	GENERATED_BODY()
 };
 
 USTRUCT()
-struct STORMSYNCTRANSPORTCORE_API FStormSyncTransportStatusPong
+struct FStormSyncTransportStatusPong
 {
 	GENERATED_BODY()
 };
diff --git a/Engine/Plugins/Experimental/StormSync/Source/StormSyncTransportCore/Public/StormSyncTransportNetworkUtils.h b/Engine/Plugins/Experimental/StormSync/Source/StormSyncTransportCore/Public/StormSyncTransportNetworkUtils.h
--- a/Engine/Plugins/Experimental/StormSync/Source/StormSyncTransportCore/Public/StormSyncTransportNetworkUtils.h
+++ b/Engine/Plugins/Experimental/StormSync/Source/StormSyncTransportCore/Public/StormSyncTransportNetworkUtils.h
@@ -2,7 +2,9 @@
 
 #pragma once
 
-#include "CoreMinimal.h"
+#include "Containers/ContainersFwd.h"
+
+class FString;
 
 /**
  * Network related helpers
diff --git a/Engine/Plugins/Experimental/StormSync/Source/StormSyncTransportCore/Public/StormSyncTransportSettings.h b/Engine/Plugins/Experimental/StormSync/Source/StormSyncTransportCore/Public/StormSyncTransportSettings.h
--- a/Engine/Plugins/Experimental/StormSync/Source/StormSyncTransportCore/Public/StormSyncTransportSettings.h
+++ b/Engine/Plugins/Experimental/StormSync/Source/StormSyncTransportCore/Public/StormSyncTransportSettings.h
@@ -2,7 +2,6 @@
 
 #pragma once
 
-#include "CoreMinimal.h"
 #include "Engine/DeveloperSettings.h"
 #include "StormSyncTransportSettings.generated.h"
 
@@ -13,7 +12,7 @@ class FStormSyncTransportSettingsDetailsCustomization;
  *
  * Handles config for transport / network related features.
  */
-UCLASS(config=Game, defaultconfig)
+UCLASS(Config=Game, DefaultConfig, DisplayName="Transport & Network")
 class STORMSYNCTRANSPORTCORE_API UStormSyncTransportSettings : public UDeveloperSettings
 {
 	GENERATED_BODY()
@@ -26,19 +25,7 @@ public:
 	
 	static const UStormSyncTransportSettings& Get();
 
-	/**
-	 * The category name for our developer settings
-	 *
-	 * @see GetCategoryName
-	 */
-	static constexpr const TCHAR* PluginCategoryName = TEXT("Storm Sync");
-
-	//~ Begin UDeveloperSettings interface
-	virtual FName GetCategoryName() const override;
 #if WITH_EDITOR
-	virtual FText GetSectionText() const override;
-	//~ End UDeveloperSettings interface
-
 	/** Helper to bring focus on transport settings in editor */
 	void OpenEditorSettingsWindow() const;
 #endif
diff --git a/Engine/Plugins/Experimental/StormSync/Source/StormSyncTransportCore/Public/StormSyncTransportTcpPackets.h b/Engine/Plugins/Experimental/StormSync/Source/StormSyncTransportCore/Public/StormSyncTransportTcpPackets.h
--- a/Engine/Plugins/Experimental/StormSync/Source/StormSyncTransportCore/Public/StormSyncTransportTcpPackets.h
+++ b/Engine/Plugins/Experimental/StormSync/Source/StormSyncTransportCore/Public/StormSyncTransportTcpPackets.h
@@ -2,7 +2,6 @@
 
 #pragma once
 
-#include "CoreMinimal.h"
 #include "StormSyncTransportTcpPackets.generated.h"
 
 USTRUCT()
diff --git a/Engine/Plugins/Experimental/StormSync/Source/StormSyncTransportServer/Private/ServiceDiscovery/StormSyncDiscoveryManager.h b/Engine/Plugins/Experimental/StormSync/Source/StormSyncTransportServer/Private/ServiceDiscovery/StormSyncDiscoveryManager.h
--- a/Engine/Plugins/Experimental/StormSync/Source/StormSyncTransportServer/Private/ServiceDiscovery/StormSyncDiscoveryManager.h
+++ b/Engine/Plugins/Experimental/StormSync/Source/StormSyncTransportServer/Private/ServiceDiscovery/StormSyncDiscoveryManager.h
@@ -2,7 +2,6 @@
 
 #pragma once
 
-#include "CoreMinimal.h"
 #include "HAL/Runnable.h"
 #include "MessageEndpoint.h"
 #include "StormSyncTransportMessages.h"
diff --git a/Engine/Plugins/Experimental/StormSync/Source/StormSyncTransportServer/Private/ServiceDiscovery/StormSyncHeartbeatEmitter.h b/Engine/Plugins/Experimental/StormSync/Source/StormSyncTransportServer/Private/ServiceDiscovery/StormSyncHeartbeatEmitter.h
--- a/Engine/Plugins/Experimental/StormSync/Source/StormSyncTransportServer/Private/ServiceDiscovery/StormSyncHeartbeatEmitter.h
+++ b/Engine/Plugins/Experimental/StormSync/Source/StormSyncTransportServer/Private/ServiceDiscovery/StormSyncHeartbeatEmitter.h
@@ -2,7 +2,6 @@
 
 #pragma once
 
-#include "CoreMinimal.h"
 #include "HAL/Runnable.h"
 #include "HAL/RunnableThread.h"
 #include "IMessageContext.h"
diff --git a/Engine/Plugins/Experimental/StormSync/Source/StormSyncTransportServer/Private/Socket/StormSyncTransportTcpServer.h b/Engine/Plugins/Experimental/StormSync/Source/StormSyncTransportServer/Private/Socket/StormSyncTransportTcpServer.h
--- a/Engine/Plugins/Experimental/StormSync/Source/StormSyncTransportServer/Private/Socket/StormSyncTransportTcpServer.h
+++ b/Engine/Plugins/Experimental/StormSync/Source/StormSyncTransportServer/Private/Socket/StormSyncTransportTcpServer.h
@@ -3,7 +3,6 @@
 #pragma once
 
 #include "Containers/Queue.h"
-#include "CoreMinimal.h"
 #include "HAL/Runnable.h"
 #include "Interfaces/IPv4/IPv4Endpoint.h"
 #include "JsonObjectConverter.h"
@@ -59,7 +58,7 @@ class FStormSyncTransportTcpServer : public FRunnable
 {
 public:
 	explicit FStormSyncTransportTcpServer(const FIPv4Address& InEndpointAddress, const uint16 InEndpointPort, const uint32 InInactiveTimeoutSeconds);
-	virtual ~FStormSyncTransportTcpServer();
+	virtual ~FStormSyncTransportTcpServer() override;
 
 	//~ Begin FRunnable interface
 	virtual bool Init() override;
diff --git a/Engine/Plugins/Experimental/StormSync/Source/StormSyncTransportServer/Private/StormSyncTransportServerEndpoint.h b/Engine/Plugins/Experimental/StormSync/Source/StormSyncTransportServer/Private/StormSyncTransportServerEndpoint.h
--- a/Engine/Plugins/Experimental/StormSync/Source/StormSyncTransportServer/Private/StormSyncTransportServerEndpoint.h
+++ b/Engine/Plugins/Experimental/StormSync/Source/StormSyncTransportServer/Private/StormSyncTransportServerEndpoint.h
@@ -2,7 +2,6 @@
 
 #pragma once
 
-#include "CoreMinimal.h"
 #include "Interfaces/IPv4/IPv4Endpoint.h"
 #include "IStormSyncTransportServerLocalEndpoint.h"
 #include "StormSyncTransportMessages.h"
@@ -22,7 +21,7 @@ class FStormSyncTransportServerEndpoint : public IStormSyncTransportServerLocalE
 {
 public:
 	FStormSyncTransportServerEndpoint();
-	virtual ~FStormSyncTransportServerEndpoint();
+	virtual ~FStormSyncTransportServerEndpoint() override;
 
 	//~ Begin IStormSyncTransportServerLocalEndpoint
 	/** Creates and starts tcp listener */
diff --git a/Engine/Plugins/Experimental/StormSync/Source/StormSyncTransportServer/Private/StormSyncTransportServerLog.h b/Engine/Plugins/Experimental/StormSync/Source/StormSyncTransportServer/Private/StormSyncTransportServerLog.h
new file mode 100644
--- /dev/null
+++ b/Engine/Plugins/Experimental/StormSync/Source/StormSyncTransportServer/Private/StormSyncTransportServerLog.h
@@ -0,0 +1,12 @@
+// Copyright Epic Games, Inc. All Rights Reserved.
+
+#pragma once
+
+#include "Logging/LogMacros.h"
+
+DECLARE_LOG_CATEGORY_EXTERN(LogStormSyncServer, Display, All);
+
+#define STORM_SYNC_SERVER_LOG(Verbosity, Format, ...) \
+{ \
+    UE_LOG(LogStormSyncServer, Verbosity, Format, ##__VA_ARGS__); \
+}
diff --git a/Engine/Plugins/Experimental/StormSync/Source/StormSyncTransportServer/Public/IStormSyncTransportServerLocalEndpoint.h b/Engine/Plugins/Experimental/StormSync/Source/StormSyncTransportServer/Public/IStormSyncTransportServerLocalEndpoint.h
--- a/Engine/Plugins/Experimental/StormSync/Source/StormSyncTransportServer/Public/IStormSyncTransportServerLocalEndpoint.h
+++ b/Engine/Plugins/Experimental/StormSync/Source/StormSyncTransportServer/Public/IStormSyncTransportServerLocalEndpoint.h
@@ -2,13 +2,12 @@
 
 #pragma once
 
-#include "CoreMinimal.h"
 #include "IStormSyncTransportLocalEndpoint.h"
 
 struct FIPv4Endpoint;
 
 /** Local Endpoint interface for Storm Sync Server */
-class STORMSYNCTRANSPORTSERVER_API IStormSyncTransportServerLocalEndpoint : public IStormSyncTransportLocalEndpoint
+class IStormSyncTransportServerLocalEndpoint : public IStormSyncTransportLocalEndpoint
 {
 public:
 	/** Returns endpoint address tcp server is currently listening on (ip:port), empty string otherwise (if not active / listening) */
diff --git a/Engine/Plugins/Experimental/StormSync/Source/StormSyncTransportServer/Public/IStormSyncTransportServerModule.h b/Engine/Plugins/Experimental/StormSync/Source/StormSyncTransportServer/Public/IStormSyncTransportServerModule.h
--- a/Engine/Plugins/Experimental/StormSync/Source/StormSyncTransportServer/Public/IStormSyncTransportServerModule.h
+++ b/Engine/Plugins/Experimental/StormSync/Source/StormSyncTransportServer/Public/IStormSyncTransportServerModule.h
@@ -2,7 +2,6 @@
 
 #pragma once
 
-#include "CoreMinimal.h"
 #include "Modules/ModuleInterface.h"
 #include "Modules/ModuleManager.h"
 
@@ -18,7 +17,7 @@ public:
 	 *
 	 * @return Returns singleton instance, loading the module on demand if needed
 	 */
-	static inline IStormSyncTransportServerModule& Get()
+	static IStormSyncTransportServerModule& Get()
 	{
 		static const FName ModuleName = "StormSyncTransportServer";
 		return FModuleManager::LoadModuleChecked<IStormSyncTransportServerModule>(ModuleName);
@@ -29,7 +28,7 @@ public:
 	 *
 	 * @return True if the module is loaded and ready to use
 	 */
-	static inline bool IsAvailable()
+	static bool IsAvailable()
 	{
 		static const FName ModuleName = "StormSyncTransportServer";
 		return FModuleManager::Get().IsModuleLoaded(ModuleName);
diff --git a/Engine/Plugins/Experimental/StormSync/Source/StormSyncTransportServer/Public/StormSyncTransportServerLog.h b/Engine/Plugins/Experimental/StormSync/Source/StormSyncTransportServer/Public/StormSyncTransportServerLog.h
deleted file mode 100644
