Storm Sync:
- Replaced custom log macros with the normal UE_LOG

#changelist validated
#virtualized

diff --git a/Engine/Plugins/Experimental/StormSync/Source/StormSyncCore/Private/StormSyncCoreLog.h b/Engine/Plugins/Experimental/StormSync/Source/StormSyncCore/Private/StormSyncCoreLog.h
--- a/Engine/Plugins/Experimental/StormSync/Source/StormSyncCore/Private/StormSyncCoreLog.h
+++ b/Engine/Plugins/Experimental/StormSync/Source/StormSyncCore/Private/StormSyncCoreLog.h
@@ -5,8 +5,3 @@
 #include "Logging/LogMacros.h"
 
 DECLARE_LOG_CATEGORY_EXTERN(LogStormSyncCore, Display, All);
-
-#define STORM_SYNC_LOG(Verbosity, Format, ...) \
-{ \
-    UE_LOG(LogStormSyncCore, Verbosity, Format, ##__VA_ARGS__); \
-}
diff --git a/Engine/Plugins/Experimental/StormSync/Source/StormSyncCore/Private/StormSyncCoreUtils.cpp b/Engine/Plugins/Experimental/StormSync/Source/StormSyncCore/Private/StormSyncCoreUtils.cpp
--- a/Engine/Plugins/Experimental/StormSync/Source/StormSyncCore/Private/StormSyncCoreUtils.cpp
+++ b/Engine/Plugins/Experimental/StormSync/Source/StormSyncCore/Private/StormSyncCoreUtils.cpp
@@ -145,7 +145,7 @@ TFuture<TArray<FStormSyncFileDependency>> FStormSyncCoreUtils::GetAvaFileDepende
 
 		if (!GetAvaFileDependenciesForPackages(PackageNames, FileDependencies, ErrorText, bInShouldValidatePackages))
 		{
-			STORM_SYNC_LOG(Error, TEXT("FStormSyncCoreUtils::GetAvaFileDependenciesAsync - Error: %s"), *ErrorText.ToString())
+			UE_LOG(LogStormSyncCore, Error, TEXT("FStormSyncCoreUtils::GetAvaFileDependenciesAsync - Error: %s"), *ErrorText.ToString());
 			return FileDependencies;
 		}
 
@@ -241,12 +241,12 @@ bool FStormSyncCoreUtils::CreatePakBuffer(const TArray<FName>& InPackageNames, T
 	int32 FileCount = PackageNames.Num();
 	*MemoryArchive << FileCount;
 
-	STORM_SYNC_LOG(Display, TEXT("FStormSyncCoreUtils::CreatePakBuffer - Creating Pak file for %d files."), PackageNames.Num());
+	UE_LOG(LogStormSyncCore, Display, TEXT("FStormSyncCoreUtils::CreatePakBuffer - Creating Pak file for %d files."), PackageNames.Num());
 
 	for (const FName& PackageName : PackageNames)
 	{
 		const FString PackageNameStr = PackageName.ToString();
-		STORM_SYNC_LOG(Verbose, TEXT("FStormSyncCoreUtils::CreatePakBuffer - Handle `%s` file to add."), *PackageNameStr);
+		UE_LOG(LogStormSyncCore, Verbose, TEXT("FStormSyncCoreUtils::CreatePakBuffer - Handle `%s` file to add."), *PackageNameStr);
 
 		FString PackageFilepath;
 		if (!FPackageName::DoesPackageExist(PackageNameStr, &PackageFilepath))
@@ -293,7 +293,7 @@ bool FStormSyncCoreUtils::CreatePakBuffer(const TArray<FName>& InPackageNames, T
 		}
 	}
 
-	STORM_SYNC_LOG(Display, TEXT("Added %d files to %s, %lld bytes total, time %.2lfs."), FileCount, *MemoryArchive->GetArchiveName(), MemoryArchive->TotalSize(), FPlatformTime::Seconds() - StartTime);
+	UE_LOG(LogStormSyncCore, Display, TEXT("Added %d files to %s, %lld bytes total, time %.2lfs."), FileCount, *MemoryArchive->GetArchiveName(), MemoryArchive->TotalSize(), FPlatformTime::Seconds() - StartTime);
 
 	MemoryArchive->Close();
 	MemoryArchive.Reset();
@@ -303,7 +303,7 @@ bool FStormSyncCoreUtils::CreatePakBuffer(const TArray<FName>& InPackageNames, T
 
 bool FStormSyncCoreUtils::ExtractPakBuffer(const TArray<uint8>& InPakBuffer, const FStormSyncCoreExtractArgs& InExtractArgs, TMap<FString, FString>& OutSuccessfullyExtractedPackages, TArray<FText>& OutErrors)
 {
-	STORM_SYNC_LOG(Display, TEXT("FStormSyncCoreUtils:ExtractPakBuffer Extracting package from buffer of size: %d"), InPakBuffer.Num());
+	UE_LOG(LogStormSyncCore, Display, TEXT("FStormSyncCoreUtils:ExtractPakBuffer Extracting package from buffer of size: %d"), InPakBuffer.Num());
 
 	const TUniquePtr<FArchive> MemoryArchive = MakeUnique<FMemoryReader>(InPakBuffer);
 	if (!MemoryArchive.IsValid())
@@ -320,7 +320,7 @@ bool FStormSyncCoreUtils::ExtractPakBuffer(const TArray<uint8>& InPakBuffer, con
 	// Notify pack extraction process is starting
 	InExtractArgs.OnPakPreExtract.ExecuteIfBound(FileCount);
 
-	STORM_SYNC_LOG(Verbose, TEXT("FStormSyncCoreUtils:ExtractPakBuffer FileCount: %d"), FileCount);
+	UE_LOG(LogStormSyncCore, Verbose, TEXT("FStormSyncCoreUtils:ExtractPakBuffer FileCount: %d"), FileCount);
 
 	for (int32 Index = 0; Index < FileCount; ++Index)
 	{
@@ -407,7 +407,7 @@ FString FStormSyncCoreUtils::GetHumanReadableByteSize(uint64 InSize)
 
 TArray<FStormSyncFileModifierInfo> FStormSyncCoreUtils::GetSyncFileModifiers(const TArray<FName>& InPackageNames, const TArray<FStormSyncFileDependency>& InRemoteDependencies)
 {
-	STORM_SYNC_LOG(Verbose, TEXT("FStormSyncCoreUtils::GetSyncFileModifiers - InPackageNames: %d, InRemoteDependencies: %d"), InPackageNames.Num(), InRemoteDependencies.Num())
+	UE_LOG(LogStormSyncCore, Verbose, TEXT("FStormSyncCoreUtils::GetSyncFileModifiers - InPackageNames: %d, InRemoteDependencies: %d"), InPackageNames.Num(), InRemoteDependencies.Num());
 
 	// If provided package names is empty, early out
 	if (InPackageNames.IsEmpty())
@@ -439,7 +439,7 @@ TArray<FStormSyncFileModifierInfo> FStormSyncCoreUtils::GetSyncFileModifiers(con
 			FStormSyncFileModifierInfo Modifier;
 			Modifier.ModifierOperation = EStormSyncModifierOperation::Missing;
 			Modifier.FileDependency = LocalDependency;
-			STORM_SYNC_LOG(Verbose, TEXT("\tAdding Modifier: %s"), *Modifier.ToString())
+			UE_LOG(LogStormSyncCore, Verbose, TEXT("\tAdding Modifier: %s"), *Modifier.ToString());
 			Modifiers.Add(Modifier);
 		}
 	}
@@ -467,7 +467,7 @@ TArray<FStormSyncFileModifierInfo> FStormSyncCoreUtils::GetSyncFileModifiers(con
 			FStormSyncFileModifierInfo Modifier;
 			Modifier.ModifierOperation = EStormSyncModifierOperation::Addition;
 			Modifier.FileDependency = RemoteDependency;
-			STORM_SYNC_LOG(Verbose, TEXT("\tAdding Modifier: %s"), *Modifier.ToString())
+			UE_LOG(LogStormSyncCore, Verbose, TEXT("\tAdding Modifier: %s"), *Modifier.ToString());
 			Modifiers.Add(Modifier);
 
 			continue;
@@ -483,19 +483,19 @@ TArray<FStormSyncFileModifierInfo> FStormSyncCoreUtils::GetSyncFileModifiers(con
 		// Check against FileSize
 		if (MatchingDependency->FileSize != RemoteDependency.FileSize)
 		{
-			STORM_SYNC_LOG(Verbose, TEXT("FStormSyncCoreUtils::GetSyncFileModifiers - Handle %s"), *RemoteDependency.PackageName.ToString())
-			STORM_SYNC_LOG(Verbose, TEXT("\tLocal: %s"), *MatchingDependency->ToString())
-			STORM_SYNC_LOG(Verbose, TEXT("\tRemote: %s"), *RemoteDependency.ToString())
-			STORM_SYNC_LOG(Verbose, TEXT("\tDirty because of mismatch filesize %lld vs %lld"), MatchingDependency->FileSize, RemoteDependency.FileSize)
+			UE_LOG(LogStormSyncCore, Verbose, TEXT("FStormSyncCoreUtils::GetSyncFileModifiers - Handle %s"), *RemoteDependency.PackageName.ToString());
+			UE_LOG(LogStormSyncCore, Verbose, TEXT("\tLocal: %s"), *MatchingDependency->ToString());
+			UE_LOG(LogStormSyncCore, Verbose, TEXT("\tRemote: %s"), *RemoteDependency.ToString());
+			UE_LOG(LogStormSyncCore, Verbose, TEXT("\tDirty because of mismatch filesize %lld vs %lld"), MatchingDependency->FileSize, RemoteDependency.FileSize);
 			bIsDirty = true;
 		}
 		// Check against File Hash
 		else if (MatchingDependency->FileHash != RemoteDependency.FileHash)
 		{
-			STORM_SYNC_LOG(Verbose, TEXT("FStormSyncCoreUtils::GetSyncFileModifiers - Handle %s"), *RemoteDependency.PackageName.ToString())
-			STORM_SYNC_LOG(Verbose, TEXT("\tLocal: %s"), *MatchingDependency->ToString())
-			STORM_SYNC_LOG(Verbose, TEXT("\tRemote: %s"), *RemoteDependency.ToString())
-			STORM_SYNC_LOG(Verbose, TEXT("\tDirty because of mismatch file hash %s vs %s"), *MatchingDependency->FileHash, *RemoteDependency.FileHash)
+			UE_LOG(LogStormSyncCore, Verbose, TEXT("FStormSyncCoreUtils::GetSyncFileModifiers - Handle %s"), *RemoteDependency.PackageName.ToString());
+			UE_LOG(LogStormSyncCore, Verbose, TEXT("\tLocal: %s"), *MatchingDependency->ToString());
+			UE_LOG(LogStormSyncCore, Verbose, TEXT("\tRemote: %s"), *RemoteDependency.ToString());
+			UE_LOG(LogStormSyncCore, Verbose, TEXT("\tDirty because of mismatch file hash %s vs %s"), *MatchingDependency->FileHash, *RemoteDependency.FileHash);
 			bIsDirty = true;
 		}
 
@@ -505,7 +505,7 @@ TArray<FStormSyncFileModifierInfo> FStormSyncCoreUtils::GetSyncFileModifiers(con
 			FStormSyncFileModifierInfo Modifier;
 			Modifier.ModifierOperation = EStormSyncModifierOperation::Overwrite;
 			Modifier.FileDependency = *MatchingDependency;
-			STORM_SYNC_LOG(Verbose, TEXT("\tAdding Modifier: %s"), *Modifier.ToString())
+			UE_LOG(LogStormSyncCore, Verbose, TEXT("\tAdding Modifier: %s"), *Modifier.ToString());
 			Modifiers.Add(Modifier);
 		}
 	}
@@ -519,14 +519,14 @@ void FStormSyncCoreUtils::RecursiveGetDependencies(const FName& InPackageName, T
 	TArray<FAssetData> Assets;
 	if (!GetAssetData(InPackageName.ToString(), Assets, Dependencies))
 	{
-		STORM_SYNC_LOG(Warning, TEXT("FStormSyncCoreUtils::RecursiveGetDependencies - GetAssetData failed to load assets for %s"), *InPackageName.ToString());
+		UE_LOG(LogStormSyncCore, Warning, TEXT("FStormSyncCoreUtils::RecursiveGetDependencies - GetAssetData failed to load assets for %s"), *InPackageName.ToString());
 		return;
 	}
 
 	for (const FName& Dependency : Dependencies)
 	{
 		const FString DependencyName = Dependency.ToString();
-		STORM_SYNC_LOG(Verbose, TEXT("FStormSyncCoreUtils::RecursiveGetDependencies - Gather dependencies for %s"), *DependencyName);
+		UE_LOG(LogStormSyncCore, Verbose, TEXT("FStormSyncCoreUtils::RecursiveGetDependencies - Gather dependencies for %s"), *DependencyName);
 
 		if (IsValidDependency(DependencyName))
 		{
@@ -541,7 +541,7 @@ void FStormSyncCoreUtils::RecursiveGetDependencies(const FName& InPackageName, T
 
 bool FStormSyncCoreUtils::IsValidDependency(const FString& InDependencyName)
 {
-	STORM_SYNC_LOG(Verbose, TEXT("FStormSyncCoreUtils::IsValidDependency - Check dependency (%s)"), *InDependencyName);
+	UE_LOG(LogStormSyncCore, Verbose, TEXT("FStormSyncCoreUtils::IsValidDependency - Check dependency (%s)"), *InDependencyName);
 
 	const UStormSyncCoreSettings* Settings = GetDefault<UStormSyncCoreSettings>();
 	check(Settings);
@@ -549,7 +549,7 @@ bool FStormSyncCoreUtils::IsValidDependency(const FString& InDependencyName)
 	// Filter out any references outside of /Game if user opted to only export /Game content
 	if (Settings->bExportOnlyGameContent && !InDependencyName.StartsWith(TEXT("/Game")))
 	{
-		STORM_SYNC_LOG(Verbose, TEXT("FStormSyncCoreUtils::IsValidDependency - Filter out dependency \"%s\". Ignored by bExportOnlyGameContent setting"), *InDependencyName);
+		UE_LOG(LogStormSyncCore, Verbose, TEXT("FStormSyncCoreUtils::IsValidDependency - Filter out dependency \"%s\". Ignored by bExportOnlyGameContent setting"), *InDependencyName);
 		return false;
 	}
 
@@ -562,7 +562,7 @@ bool FStormSyncCoreUtils::IsValidDependency(const FString& InDependencyName)
 	{
 		if (InDependencyName.StartsWith(IgnoredPackageName.ToString()))
 		{
-			STORM_SYNC_LOG(Verbose, TEXT("FStormSyncCoreUtils::IsValidDependency - Filter out dependency \"%s\". Ignored by \"%s\" IgnoredPackages setting"), *InDependencyName, *IgnoredPackageName.ToString());
+			UE_LOG(LogStormSyncCore, Verbose, TEXT("FStormSyncCoreUtils::IsValidDependency - Filter out dependency \"%s\". Ignored by \"%s\" IgnoredPackages setting"), *InDependencyName, *IgnoredPackageName.ToString());
 			return false;
 		}
 	}
diff --git a/Engine/Plugins/Experimental/StormSync/Source/StormSyncDrives/Private/StormSyncDrivesLog.h b/Engine/Plugins/Experimental/StormSync/Source/StormSyncDrives/Private/StormSyncDrivesLog.h
--- a/Engine/Plugins/Experimental/StormSync/Source/StormSyncDrives/Private/StormSyncDrivesLog.h
+++ b/Engine/Plugins/Experimental/StormSync/Source/StormSyncDrives/Private/StormSyncDrivesLog.h
@@ -5,8 +5,3 @@
 #include "Logging/LogMacros.h"
 
 DECLARE_LOG_CATEGORY_EXTERN(LogStormSyncDrives, Display, All);
-
-#define STORM_SYNC_DRIVES_LOG(Verbosity, Format, ...) \
-{ \
-    UE_LOG(LogStormSyncDrives, Verbosity, Format, ##__VA_ARGS__); \
-}
diff --git a/Engine/Plugins/Experimental/StormSync/Source/StormSyncDrives/Private/StormSyncDrivesModule.cpp b/Engine/Plugins/Experimental/StormSync/Source/StormSyncDrives/Private/StormSyncDrivesModule.cpp
--- a/Engine/Plugins/Experimental/StormSync/Source/StormSyncDrives/Private/StormSyncDrivesModule.cpp
+++ b/Engine/Plugins/Experimental/StormSync/Source/StormSyncDrives/Private/StormSyncDrivesModule.cpp
@@ -59,7 +59,7 @@ void FStormSyncDrivesModule::ShutdownModule()
 
 bool FStormSyncDrivesModule::RegisterMountPoint(const FStormSyncMountPointConfig& InMountPoint, FText& ErrorText)
 {
-	STORM_SYNC_DRIVES_LOG(Display, TEXT("FStormSyncDrivesModule::RegisterMountPoint ... MountPoint: %s, Path: %s"), *InMountPoint.MountPoint, *InMountPoint.MountDirectory.Path)
+	UE_LOG(LogStormSyncDrives, Display, TEXT("FStormSyncDrivesModule::RegisterMountPoint ... MountPoint: %s, Path: %s"), *InMountPoint.MountPoint, *InMountPoint.MountDirectory.Path);
 
 	// Validate first configuration of the mount point
 	FText ValidationError;
@@ -92,7 +92,7 @@ bool FStormSyncDrivesModule::RegisterMountPoint(const FStormSyncMountPointConfig
 		return false;
 	}
 
-	STORM_SYNC_DRIVES_LOG(Display, TEXT("FStormSyncDrivesModule::RegisterMountPoint ... Try register %s to %s"), *MountPoint, *InMountPoint.MountDirectory.Path)
+	UE_LOG(LogStormSyncDrives, Display, TEXT("FStormSyncDrivesModule::RegisterMountPoint ... Try register %s to %s"), *MountPoint, *InMountPoint.MountDirectory.Path);
 	FPackageName::RegisterMountPoint(MountPoint, InMountPoint.MountDirectory.Path);
 	return true;
 }
@@ -120,8 +120,8 @@ bool FStormSyncDrivesModule::UnregisterMountPoint(const FStormSyncMountPointConf
 
 void FStormSyncDrivesModule::OnEngineLoopInitComplete()
 {
-	STORM_SYNC_DRIVES_LOG(Display, TEXT("FStormSyncDrivesModule::OnEngineLoopInitComplete ..."))
-	STORM_SYNC_DRIVES_LOG(Display, TEXT("\t Mounting Drives based on config"))
+	UE_LOG(LogStormSyncDrives, Display, TEXT("FStormSyncDrivesModule::OnEngineLoopInitComplete ..."));
+	UE_LOG(LogStormSyncDrives, Display, TEXT("\t Mounting Drives based on config"));
 
 	ResetMountedDrivesFromSettings(GetDefault<UStormSyncDrivesSettings>());
 }
@@ -155,7 +155,7 @@ bool FStormSyncDrivesModule::ValidateMountPoint(const FStormSyncMountPointConfig
 			ErrorText
 		);
 		AddMessageError(ErrorLog);
-		STORM_SYNC_DRIVES_LOG(Error, TEXT("FStormSyncDrivesModule::ValidateMountPoint - %s"), *ErrorLog.ToString())
+		UE_LOG(LogStormSyncDrives, Error, TEXT("FStormSyncDrivesModule::ValidateMountPoint - %s"), *ErrorLog.ToString());
 		return false;
 	}
 
@@ -178,7 +178,7 @@ bool FStormSyncDrivesModule::ValidateMountPoints(const TArray<FStormSyncMountPoi
 	int32 Index = 0;
 	for (const FStormSyncMountPointConfig& Entry : InMountPoints)
 	{
-		STORM_SYNC_DRIVES_LOG(Display, TEXT("Try validate entry %s %s"), *Entry.MountPoint, *Entry.MountDirectory.Path)
+		UE_LOG(LogStormSyncDrives, Display, TEXT("Try validate entry %s %s"), *Entry.MountPoint, *Entry.MountDirectory.Path);
 
 		// First check for mount point validity individually
 		bResult &= ValidateMountPoint(Entry, Index);
@@ -195,7 +195,7 @@ bool FStormSyncDrivesModule::ValidateMountPoints(const TArray<FStormSyncMountPoi
 		for (const FText& ValidationError : ValidationErrors)
 		{
 			AddMessageError(ValidationError);
-			STORM_SYNC_DRIVES_LOG(Error, TEXT("FStormSyncDrivesModule::ValidateMountPoints - %s"), *ValidationError.ToString())
+			UE_LOG(LogStormSyncDrives, Error, TEXT("FStormSyncDrivesModule::ValidateMountPoints - %s"), *ValidationError.ToString());
 		}
 	}
 
@@ -212,7 +212,7 @@ bool FStormSyncDrivesModule::ValidateMountPoints(const TArray<FStormSyncMountPoi
 
 void FStormSyncDrivesModule::ResetMountedDrivesFromSettings(const UStormSyncDrivesSettings* InSettings)
 {
-	STORM_SYNC_DRIVES_LOG(Display, TEXT("FStormSyncDrivesModule::ResetMountedDrives ..."))
+	UE_LOG(LogStormSyncDrives, Display, TEXT("FStormSyncDrivesModule::ResetMountedDrives ..."));
 	check(InSettings);
 
 	if (ValidateMountPoints(InSettings->MountPoints))
@@ -225,7 +225,7 @@ void FStormSyncDrivesModule::ResetMountedDrivesFromSettings(const UStormSyncDriv
 
 void FStormSyncDrivesModule::UnregisterMountedDrives()
 {
-	STORM_SYNC_DRIVES_LOG(Display, TEXT("FStormSyncDrivesModule::UnregisterMountedDrives ..."))
+	UE_LOG(LogStormSyncDrives, Display, TEXT("FStormSyncDrivesModule::UnregisterMountedDrives ..."));
 	for (const TSharedPtr<FStormSyncMountPointConfig>& MountedDrive : MountedDrives)
 	{
 		if (!MountedDrive.IsValid())
@@ -236,14 +236,14 @@ void FStormSyncDrivesModule::UnregisterMountedDrives()
 		FText ErrorText;
 		if (!UnregisterMountPoint(*MountedDrive.Get(), ErrorText))
 		{
-			STORM_SYNC_DRIVES_LOG(Error, TEXT("FStormSyncDrivesModule::UnregisterMountedDrives failed with Error: %s"), *ErrorText.ToString())
+			UE_LOG(LogStormSyncDrives, Error, TEXT("FStormSyncDrivesModule::UnregisterMountedDrives failed with Error: %s"), *ErrorText.ToString());
 		}
 	}
 }
 
 void FStormSyncDrivesModule::CacheMountedDrives(const TArray<FStormSyncMountPointConfig>& InMountPoints)
 {
-	STORM_SYNC_DRIVES_LOG(Display, TEXT("FStormSyncDrivesModule::CacheMountedDrives ..."))
+	UE_LOG(LogStormSyncDrives, Display, TEXT("FStormSyncDrivesModule::CacheMountedDrives ..."));
 	MountedDrives.Reset();
 
 	for (const FStormSyncMountPointConfig& Entry : InMountPoints)
@@ -254,7 +254,7 @@ void FStormSyncDrivesModule::CacheMountedDrives(const TArray<FStormSyncMountPoin
 
 void FStormSyncDrivesModule::RegisterMountedDrives()
 {
-	STORM_SYNC_DRIVES_LOG(Display, TEXT("FStormSyncDrivesModule::RegisterMountedDrives ..."))
+	UE_LOG(LogStormSyncDrives, Display, TEXT("FStormSyncDrivesModule::RegisterMountedDrives ..."));
 	for (TSharedPtr<FStormSyncMountPointConfig> MountedDrive : MountedDrives)
 	{
 		if (!MountedDrive.IsValid())
@@ -265,7 +265,7 @@ void FStormSyncDrivesModule::RegisterMountedDrives()
 		FText ErrorText;
 		if (!RegisterMountPoint(*MountedDrive.Get(), ErrorText))
 		{
-			STORM_SYNC_DRIVES_LOG(Error, TEXT("FStormSyncDrivesModule::RegisterMountedDrives failed with Error: %s"), *ErrorText.ToString())
+			UE_LOG(LogStormSyncDrives, Error, TEXT("FStormSyncDrivesModule::RegisterMountedDrives failed with Error: %s"), *ErrorText.ToString());
 		}
 	}
 }
diff --git a/Engine/Plugins/Experimental/StormSync/Source/StormSyncEditor/Private/ContextMenu/StormSyncAssetFolderContextMenu.cpp b/Engine/Plugins/Experimental/StormSync/Source/StormSyncEditor/Private/ContextMenu/StormSyncAssetFolderContextMenu.cpp
--- a/Engine/Plugins/Experimental/StormSync/Source/StormSyncEditor/Private/ContextMenu/StormSyncAssetFolderContextMenu.cpp
+++ b/Engine/Plugins/Experimental/StormSync/Source/StormSyncEditor/Private/ContextMenu/StormSyncAssetFolderContextMenu.cpp
@@ -29,13 +29,13 @@
 
 void FStormSyncAssetFolderContextMenu::Initialize()
 {
-	STORM_SYNC_EDITOR_LOG(Verbose, TEXT("FStormSyncAssetFolderContextMenu::Initialize - ..."));
+	UE_LOG(LogStormSyncEditor, Verbose, TEXT("FStormSyncAssetFolderContextMenu::Initialize - ..."));
 	UToolMenus::RegisterStartupCallback(FSimpleMulticastDelegate::FDelegate::CreateSP(this, &FStormSyncAssetFolderContextMenu::RegisterMenus));
 }
 
 void FStormSyncAssetFolderContextMenu::Shutdown()
 {
-	STORM_SYNC_EDITOR_LOG(Verbose, TEXT("FStormSyncAssetFolderContextMenu::Shutdown - ..."));
+	UE_LOG(LogStormSyncEditor, Verbose, TEXT("FStormSyncAssetFolderContextMenu::Shutdown - ..."));
 }
 
 void FStormSyncAssetFolderContextMenu::BuildPushAssetsMenuSection(FMenuBuilder& InMenuBuilder, const TArray<FName> InPackageNames, const bool bInIsPushing)
@@ -108,7 +108,7 @@ TArray<FAssetData> FStormSyncAssetFolderContextMenu::GetDirtyAssets(const TArray
 
 		if (!AssetRegistry.GetAssetsByPackageName(PackageName, Assets))
 		{
-			STORM_SYNC_EDITOR_LOG(Warning, TEXT("FStormSyncAssetFolderContextMenu::HasDirtyAssets - GetAssetsByPackageName failed to load assets for %s"), *PackageName.ToString());
+			UE_LOG(LogStormSyncEditor, Warning, TEXT("FStormSyncAssetFolderContextMenu::HasDirtyAssets - GetAssetsByPackageName failed to load assets for %s"), *PackageName.ToString());
 			continue;
 		}
 
@@ -157,7 +157,7 @@ TArray<FAssetData> FStormSyncAssetFolderContextMenu::GetDirtyAssets(const TArray
 
 void FStormSyncAssetFolderContextMenu::RegisterMenus()
 {
-	STORM_SYNC_EDITOR_LOG(Verbose, TEXT("FStormSyncAssetFolderContextMenu::RegisterMenus"));
+	UE_LOG(LogStormSyncEditor, Verbose, TEXT("FStormSyncAssetFolderContextMenu::RegisterMenus"));
 
 	if (UToolMenu* Menu = UToolMenus::Get()->ExtendMenu("ContentBrowser.FolderContextMenu"))
 	{
@@ -186,7 +186,7 @@ void FStormSyncAssetFolderContextMenu::RegisterMenus()
 
 void FStormSyncAssetFolderContextMenu::PopulateAssetFolderContextMenu(UToolMenu* InMenu) const
 {
-	STORM_SYNC_EDITOR_LOG(Verbose, TEXT("FStormSyncAssetFolderContextMenu::PopulateAssetFolderContextMenu - InMenu: %s"), *GetNameSafe(InMenu));
+	UE_LOG(LogStormSyncEditor, Verbose, TEXT("FStormSyncAssetFolderContextMenu::PopulateAssetFolderContextMenu - InMenu: %s"), *GetNameSafe(InMenu));
 	check(InMenu);
 
 	const UContentBrowserDataMenuContext_FolderMenu* ContextObject = InMenu->FindContext<UContentBrowserDataMenuContext_FolderMenu>();
@@ -197,17 +197,17 @@ void FStormSyncAssetFolderContextMenu::PopulateAssetFolderContextMenu(UToolMenu*
 	TArray<FString> SelectedAssetPackages;
 	GetSelectedFilesAndFolders(InMenu, SelectedPackagePaths, SelectedAssetPackages);
 
-	STORM_SYNC_EDITOR_LOG(Verbose, TEXT("FStormSyncAssetFolderContextMenu::PopulateAssetFolderContextMenu - Result ..."));
-	STORM_SYNC_EDITOR_LOG(Verbose, TEXT("\tSelectedPackagePaths: %d"), SelectedPackagePaths.Num());
+	UE_LOG(LogStormSyncEditor, Verbose, TEXT("FStormSyncAssetFolderContextMenu::PopulateAssetFolderContextMenu - Result ..."));
+	UE_LOG(LogStormSyncEditor, Verbose, TEXT("\tSelectedPackagePaths: %d"), SelectedPackagePaths.Num());
 	for (const FString& SelectedPackagePath : SelectedPackagePaths)
 	{
-		STORM_SYNC_EDITOR_LOG(Verbose, TEXT("\t- %s"), *SelectedPackagePath);
+		UE_LOG(LogStormSyncEditor, Verbose, TEXT("\t- %s"), *SelectedPackagePath);
 	}
 
-	STORM_SYNC_EDITOR_LOG(Verbose, TEXT("\tSelectedAssetPackages: %d"), SelectedAssetPackages.Num());
+	UE_LOG(LogStormSyncEditor, Verbose, TEXT("\tSelectedAssetPackages: %d"), SelectedAssetPackages.Num());
 	for (const FString& SelectedAssetPackage : SelectedAssetPackages)
 	{
-		STORM_SYNC_EDITOR_LOG(Verbose, TEXT("\t- %s"), *SelectedAssetPackage);
+		UE_LOG(LogStormSyncEditor, Verbose, TEXT("\t- %s"), *SelectedAssetPackage);
 	}
 	
 	AddFolderMenuOptions(InMenu, SelectedPackagePaths, SelectedAssetPackages);
@@ -215,7 +215,7 @@ void FStormSyncAssetFolderContextMenu::PopulateAssetFolderContextMenu(UToolMenu*
 
 void FStormSyncAssetFolderContextMenu::PopulateAssetFileContextMenu(UToolMenu* InMenu)
 {
-	STORM_SYNC_EDITOR_LOG(Verbose, TEXT("FStormSyncAssetFolderContextMenu::PopulateAssetFileContextMenu - InMenu: %s"), *GetNameSafe(InMenu));
+	UE_LOG(LogStormSyncEditor, Verbose, TEXT("FStormSyncAssetFolderContextMenu::PopulateAssetFileContextMenu - InMenu: %s"), *GetNameSafe(InMenu));
 	check(InMenu);
 
 	const UContentBrowserDataMenuContext_FileMenu* ContextObject = InMenu->FindContext<UContentBrowserDataMenuContext_FileMenu>();
@@ -225,11 +225,11 @@ void FStormSyncAssetFolderContextMenu::PopulateAssetFileContextMenu(UToolMenu* I
 	TArray<FName> SelectedAssetPackages;
 	GetSelectedFiles(InMenu, SelectedAssetPackages);
 
-	STORM_SYNC_EDITOR_LOG(Verbose, TEXT("FStormSyncAssetFolderContextMenu::PopulateAssetFileContextMenu - Result ..."));
-	STORM_SYNC_EDITOR_LOG(Verbose, TEXT("\tSelectedAssetPackages: %d"), SelectedAssetPackages.Num());
+	UE_LOG(LogStormSyncEditor, Verbose, TEXT("FStormSyncAssetFolderContextMenu::PopulateAssetFileContextMenu - Result ..."));
+	UE_LOG(LogStormSyncEditor, Verbose, TEXT("\tSelectedAssetPackages: %d"), SelectedAssetPackages.Num());
 	for (const FName& SelectedPackageName : SelectedAssetPackages)
 	{
-		STORM_SYNC_EDITOR_LOG(Verbose, TEXT("\t- %s"), *SelectedPackageName.ToString());
+		UE_LOG(LogStormSyncEditor, Verbose, TEXT("\t- %s"), *SelectedPackageName.ToString());
 	}
 
 	AddFileMenuOptions(InMenu, SelectedAssetPackages);
@@ -440,7 +440,7 @@ void FStormSyncAssetFolderContextMenu::GetSelectedAssetsInPaths(const TArray<FSt
 
 void FStormSyncAssetFolderContextMenu::GetSelectedFilesAndFolders(const UToolMenu* InMenu, TArray<FString>& OutSelectedPackagePaths, TArray<FString>& OutSelectedAssetPackages) const
 {
-	STORM_SYNC_EDITOR_LOG(Verbose, TEXT("FStormSyncAssetFolderContextMenu::GetSelectedFilesAndFolders - InMenu: %s"), *GetNameSafe(InMenu));
+	UE_LOG(LogStormSyncEditor, Verbose, TEXT("FStormSyncAssetFolderContextMenu::GetSelectedFilesAndFolders - InMenu: %s"), *GetNameSafe(InMenu));
 	check(InMenu);
 
 	const UContentBrowserDataMenuContext_FolderMenu* ContextObject = InMenu->FindContext<UContentBrowserDataMenuContext_FolderMenu>();
@@ -477,7 +477,7 @@ void FStormSyncAssetFolderContextMenu::GetSelectedFilesAndFolders(const UToolMen
 
 void FStormSyncAssetFolderContextMenu::GetSelectedFiles(const UToolMenu* InMenu, TArray<FString>& OutSelectedAssetPackages) const
 {
-	STORM_SYNC_EDITOR_LOG(Verbose, TEXT("FStormSyncAssetFolderContextMenu::GetSelectedFiles - InMenu: %s"), *GetNameSafe(InMenu));
+	UE_LOG(LogStormSyncEditor, Verbose, TEXT("FStormSyncAssetFolderContextMenu::GetSelectedFiles - InMenu: %s"), *GetNameSafe(InMenu));
 	check(InMenu);
 
 	const UContentBrowserDataMenuContext_FileMenu* ContextObject = InMenu->FindContext<UContentBrowserDataMenuContext_FileMenu>();
@@ -554,7 +554,7 @@ void FStormSyncAssetFolderContextMenu::BuildPushAssetsMenuEntries(FMenuBuilder&
 		FMessageAddress ClientAddress;
 		if (!FMessageAddress::Parse(ClientAddressId, ClientAddress))
 		{
-			STORM_SYNC_EDITOR_LOG(Error, TEXT("FStormSyncAssetFolderContextMenu::BuildPushAssetsMenuEntries - Error parsing message address for storm sync client with id \"%s\""), *ClientAddressId);
+			UE_LOG(LogStormSyncEditor, Error, TEXT("FStormSyncAssetFolderContextMenu::BuildPushAssetsMenuEntries - Error parsing message address for storm sync client with id \"%s\""), *ClientAddressId);
 			continue;
 		}
 
@@ -603,10 +603,10 @@ void FStormSyncAssetFolderContextMenu::BuildPushAssetsMenuEntries(FMenuBuilder&
 
 void FStormSyncAssetFolderContextMenu::ExecuteSyncAssetsAction(TArray<FName> InPackageNames)
 {
-	STORM_SYNC_EDITOR_LOG(Display, TEXT("FStormSyncAssetFolderContextMenu::ExecuteSyncAssetsAction - Package names %d"), InPackageNames.Num());
+	UE_LOG(LogStormSyncEditor, Display, TEXT("FStormSyncAssetFolderContextMenu::ExecuteSyncAssetsAction - Package names %d"), InPackageNames.Num());
 	for (FName PackageName : InPackageNames)
 	{
-		STORM_SYNC_EDITOR_LOG(Display, TEXT("\t- Package name: %s"), *PackageName.ToString());
+		UE_LOG(LogStormSyncEditor, Display, TEXT("\t- Package name: %s"), *PackageName.ToString());
 	}
 
 	// Note: We sync with a dummy package descriptor, next iterations could add in there an additional UI step.
@@ -617,12 +617,12 @@ void FStormSyncAssetFolderContextMenu::ExecuteSyncAssetsAction(TArray<FName> InP
 
 void FStormSyncAssetFolderContextMenu::ExecutePushAssetsAction(TArray<FName> InPackageNames, FString InMessageAddressId)
 {
-	STORM_SYNC_EDITOR_LOG(Display, TEXT("FStormSyncAssetFolderContextMenu::ExecutePushAssetsAction to %s - Package names %d"), *InMessageAddressId, InPackageNames.Num());
+	UE_LOG(LogStormSyncEditor, Display, TEXT("FStormSyncAssetFolderContextMenu::ExecutePushAssetsAction to %s - Package names %d"), *InMessageAddressId, InPackageNames.Num());
 
 	FMessageAddress RemoteMessageAddress;
 	if (!FMessageAddress::Parse(InMessageAddressId, RemoteMessageAddress))
 	{
-		STORM_SYNC_EDITOR_LOG(Error, TEXT("FStormSyncAssetFolderContextMenu::ExecutePushAssetsAction to %s - Unable to parse into a Message Address"), *InMessageAddressId);
+		UE_LOG(LogStormSyncEditor, Error, TEXT("FStormSyncAssetFolderContextMenu::ExecutePushAssetsAction to %s - Unable to parse into a Message Address"), *InMessageAddressId);
 		return;
 	}
 
@@ -643,7 +643,7 @@ void FStormSyncAssetFolderContextMenu::ExecutePushAssetsAction(TArray<FName> InP
 	const FOnStormSyncPushComplete Delegate = FOnStormSyncPushComplete::CreateLambda([](const TSharedPtr<FStormSyncTransportPushResponse>& Response)
 	{
 		check(Response.IsValid())
-		STORM_SYNC_EDITOR_LOG(Display, TEXT("FStormSyncAssetFolderContextMenu::ExecutePushAssetsAction - Got a response: %s"), *Response->ToString());
+		UE_LOG(LogStormSyncEditor, Display, TEXT("FStormSyncAssetFolderContextMenu::ExecutePushAssetsAction - Got a response: %s"), *Response->ToString());
 		UStormSyncNotificationSubsystem::Get().HandlePushResponse(Response);
 	});
 	
@@ -653,12 +653,12 @@ void FStormSyncAssetFolderContextMenu::ExecutePushAssetsAction(TArray<FName> InP
 
 void FStormSyncAssetFolderContextMenu::ExecutePullAssetsAction(TArray<FName> InPackageNames, FString InMessageAddressId)
 {
-	STORM_SYNC_EDITOR_LOG(Display, TEXT("FStormSyncAssetFolderContextMenu::ExecutePullAssetsAction to %s - Package names %d"), *InMessageAddressId, InPackageNames.Num());
+	UE_LOG(LogStormSyncEditor, Display, TEXT("FStormSyncAssetFolderContextMenu::ExecutePullAssetsAction to %s - Package names %d"), *InMessageAddressId, InPackageNames.Num());
 
 	FMessageAddress RemoteMessageAddress;
 	if (!FMessageAddress::Parse(InMessageAddressId, RemoteMessageAddress))
 	{
-		STORM_SYNC_EDITOR_LOG(Error, TEXT("FStormSyncAssetFolderContextMenu::ExecutePullAssetsAction to %s - Unable to parse into a Message Address"), *InMessageAddressId);
+		UE_LOG(LogStormSyncEditor, Error, TEXT("FStormSyncAssetFolderContextMenu::ExecutePullAssetsAction to %s - Unable to parse into a Message Address"), *InMessageAddressId);
 		return;
 	}
 
@@ -679,7 +679,7 @@ void FStormSyncAssetFolderContextMenu::ExecutePullAssetsAction(TArray<FName> InP
 	const FOnStormSyncPullComplete Delegate = FOnStormSyncPullComplete::CreateLambda([](const TSharedPtr<FStormSyncTransportPullResponse>& Response)
 	{
 		check(Response.IsValid())
-		STORM_SYNC_EDITOR_LOG(Display, TEXT("FStormSyncAssetFolderContextMenu::ExecutePullAssetsAction - Got a response: %s"), *Response->ToString());
+		UE_LOG(LogStormSyncEditor, Display, TEXT("FStormSyncAssetFolderContextMenu::ExecutePullAssetsAction - Got a response: %s"), *Response->ToString());
 		UStormSyncNotificationSubsystem::Get().HandlePullResponse(Response);
 	});
 	
@@ -689,18 +689,18 @@ void FStormSyncAssetFolderContextMenu::ExecutePullAssetsAction(TArray<FName> InP
 
 void FStormSyncAssetFolderContextMenu::ExecuteCompareWithAction(const TArray<FName> InPackageNames, const FString InMessageAddressId)
 {
-	STORM_SYNC_EDITOR_LOG(Display, TEXT("FStormSyncAssetFolderContextMenu::ExecuteCompareWithAction to %s - Package names %d"), *InMessageAddressId, InPackageNames.Num());
+	UE_LOG(LogStormSyncEditor, Display, TEXT("FStormSyncAssetFolderContextMenu::ExecuteCompareWithAction to %s - Package names %d"), *InMessageAddressId, InPackageNames.Num());
 
 	FMessageAddress RemoteAddress;
 	if (!FMessageAddress::Parse(InMessageAddressId, RemoteAddress))
 	{
-		STORM_SYNC_EDITOR_LOG(Error, TEXT("FStormSyncAssetFolderContextMenu::ExecuteCompareWithAction to %s - Unable to parse into a Message Address"), *InMessageAddressId);
+		UE_LOG(LogStormSyncEditor, Error, TEXT("FStormSyncAssetFolderContextMenu::ExecuteCompareWithAction to %s - Unable to parse into a Message Address"), *InMessageAddressId);
 		return;
 	}
 
 	const FOnStormSyncRequestStatusComplete DoneDelegate = FOnStormSyncRequestStatusComplete::CreateLambda([](const TSharedPtr<FStormSyncTransportStatusResponse>& Response)
 	{
-		STORM_SYNC_EDITOR_LOG(Display, TEXT("FStormSyncAssetFolderContextMenu::ExecuteCompareWithAction - Received response! %s"), *Response->ToString());
+		UE_LOG(LogStormSyncEditor, Display, TEXT("FStormSyncAssetFolderContextMenu::ExecuteCompareWithAction - Received response! %s"), *Response->ToString());
 
 		// Prompt the user displaying all assets that are going to be migrated
 		SStormSyncStatusWidget::OpenDialog(Response.ToSharedRef());
@@ -711,7 +711,7 @@ void FStormSyncAssetFolderContextMenu::ExecuteCompareWithAction(const TArray<FNa
 
 void FStormSyncAssetFolderContextMenu::ExecuteExportAction(const TArray<FName> InPackageNames)
 {
-	STORM_SYNC_EDITOR_LOG(Display, TEXT("FStormSyncAssetFolderContextMenu::ExecuteExportAction - Package names %d"), InPackageNames.Num());
+	UE_LOG(LogStormSyncEditor, Display, TEXT("FStormSyncAssetFolderContextMenu::ExecuteExportAction - Package names %d"), InPackageNames.Num());
 
 	// Form a full list of packages to move by including the dependencies of the supplied packages
 	TArray<FName> AllPackageNamesToMove;
@@ -724,7 +724,7 @@ void FStormSyncAssetFolderContextMenu::ExecuteExportAction(const TArray<FName> I
 		FText ErrorText;
 		if (!FStormSyncCoreUtils::GetDependenciesForPackages(InPackageNames, AllPackageNamesToMove, ErrorText))
 		{
-			STORM_SYNC_EDITOR_LOG(Error, TEXT("FStormSyncAssetFolderContextMenu::ExecuteExportAction - Failed to gather dependencies: %s"), *ErrorText.ToString())
+			UE_LOG(LogStormSyncEditor, Error, TEXT("FStormSyncAssetFolderContextMenu::ExecuteExportAction - Failed to gather dependencies: %s"), *ErrorText.ToString());
 			WarnOnInvalidPackages(InPackageNames, LOCTEXT("Export_Warning_Notify_Heading", "Some files are invalid (not existing on disk)"));
 			return;
 		}
@@ -771,7 +771,7 @@ void FStormSyncAssetFolderContextMenu::OnExportWizardCompleted(const TArray<FNam
 	const FText ErrorHeadingText = LOCTEXT("Export_Error_Notify_Heading", "Error during export, some files were invalid (not existing on disk)");
 	if (!WarnOnInvalidPackages(InPackageNames, ErrorHeadingText, FAppStyle::GetBrush(TEXT("AssetEditor.CompileStatus.Overlay.Error"))))
 	{
-		STORM_SYNC_EDITOR_LOG(Error, TEXT("FStormSyncAssetFolderContextMenu::OnExportWizardCompleted - Some files are invalid, preventing export."))
+		UE_LOG(LogStormSyncEditor, Error, TEXT("FStormSyncAssetFolderContextMenu::OnExportWizardCompleted - Some files are invalid, preventing export."));
 		// Note: Opting to prevent operation by design when Settings' bFilterInvalidReferences is false (but use it to display warnings),
 		// this can arguably be changed to not prevent operations (pak buffer creation below) by filtering invalid refs no matter what (even if bFilterInvalidReferences
 		// is false)
@@ -795,7 +795,7 @@ void FStormSyncAssetFolderContextMenu::OnExportWizardCompleted(const TArray<FNam
 		TArray<FStormSyncFileDependency> PackedFiles;
 		if (!FStormSyncCoreUtils::CreatePakBuffer(InPackageNames, Buffer, ErrorText, Delegate))
 		{
-			STORM_SYNC_EDITOR_LOG(Error, TEXT("FStormSyncAssetFolderContextMenu::OnExportWizardCompleted - CreatePakBufferWithDependencies Error: %s"), *ErrorText.ToString())
+			UE_LOG(LogStormSyncEditor, Error, TEXT("FStormSyncAssetFolderContextMenu::OnExportWizardCompleted - CreatePakBufferWithDependencies Error: %s"), *ErrorText.ToString());
 			return;
 		}
 	}
@@ -805,7 +805,7 @@ void FStormSyncAssetFolderContextMenu::OnExportWizardCompleted(const TArray<FNam
 		const TUniquePtr<FArchive> FileWriter(IFileManager::Get().CreateFileWriter(*InFilepath));
 		if (!FileWriter.IsValid())
 		{
-			STORM_SYNC_EDITOR_LOG(Error, TEXT("FStormSyncAssetFolderContextMenu::OnExportWizardCompleted - Error creating file writer"))
+			UE_LOG(LogStormSyncEditor, Error, TEXT("FStormSyncAssetFolderContextMenu::OnExportWizardCompleted - Error creating file writer"));
 			return;
 		}
 
@@ -814,7 +814,7 @@ void FStormSyncAssetFolderContextMenu::OnExportWizardCompleted(const TArray<FNam
 	}
 
 	// Notify user
-	STORM_SYNC_EDITOR_LOG(Display, TEXT("FStormSyncAssetFolderContextMenu::OnExportWizardCompleted - Created file %s"), *InFilepath)
+	UE_LOG(LogStormSyncEditor, Display, TEXT("FStormSyncAssetFolderContextMenu::OnExportWizardCompleted - Created file %s"), *InFilepath);
 
 	const FString DirectoryPath = FPaths::GetPath(InFilepath);
 	const FString BaseFilename = FPaths::GetBaseFilename(InFilepath);
@@ -831,12 +831,12 @@ void FStormSyncAssetFolderContextMenu::OnExportWizardCompleted(const TArray<FNam
 	Info.Hyperlink = FSimpleDelegate::CreateLambda([AbsolutePath]()
 	{
 		const FString HyperlinkTarget = TEXT("file://") / AbsolutePath;
-		STORM_SYNC_EDITOR_LOG(
+		UE_LOG(LogStormSyncEditor, 
 			Verbose,
 			TEXT("FStormSyncAssetFolderContextMenu::OnExportWizardCompleted - Clicked on \"Open Folder\" link with Directory: %s (HyperlinkTarget: %s)"),
 			*AbsolutePath,
 			*HyperlinkTarget
-		)
+		);
 
 		FPlatformProcess::LaunchURL(*HyperlinkTarget, nullptr, nullptr);
 	});
@@ -869,7 +869,7 @@ bool FStormSyncAssetFolderContextMenu::WarnOnInvalidPackages(const TArray<FName>
 			FText::FromString(ErrorsFileList)
 		);
 			
-		STORM_SYNC_EDITOR_LOG(Warning, TEXT("FStormSyncAssetFolderContextMenu::WarnOnInvalidPackages - %s"), *NotificationText.ToString())
+		UE_LOG(LogStormSyncEditor, Warning, TEXT("FStormSyncAssetFolderContextMenu::WarnOnInvalidPackages - %s"), *NotificationText.ToString());
 
 		FNotificationInfo NotifyInfo(NotificationText);
 		NotifyInfo.WidthOverride = 480.f;
diff --git a/Engine/Plugins/Experimental/StormSync/Source/StormSyncEditor/Private/Factories/StormSyncArchiveFactory.cpp b/Engine/Plugins/Experimental/StormSync/Source/StormSyncEditor/Private/Factories/StormSyncArchiveFactory.cpp
--- a/Engine/Plugins/Experimental/StormSync/Source/StormSyncEditor/Private/Factories/StormSyncArchiveFactory.cpp
+++ b/Engine/Plugins/Experimental/StormSync/Source/StormSyncEditor/Private/Factories/StormSyncArchiveFactory.cpp
@@ -38,7 +38,7 @@ bool UStormSyncArchiveFactory::FactoryCanImport(const FString& Filename)
 
 UObject* UStormSyncArchiveFactory::ImportObject(UClass* InClass, UObject* InOuter, FName InName, EObjectFlags Flags, const FString& Filename, const TCHAR* Params, bool& OutCanceled)
 {
-	STORM_SYNC_EDITOR_LOG(Display, TEXT("UStormSyncArchiveFactory::ImportObject InClass: %s, InOuter: %s, Filename: %s"), *GetNameSafe(InClass), *GetNameSafe(InOuter), *InName.ToString(), *Filename);
+	UE_LOG(LogStormSyncEditor, Display, TEXT("UStormSyncArchiveFactory::ImportObject InClass: %s, InOuter: %s, Filename: %s"), *GetNameSafe(InClass), *GetNameSafe(InOuter), *InName.ToString(), *Filename);
 
 	// Import Object override to customize how import is done and account for this factory operating on a dummy imported object.
 
@@ -54,11 +54,11 @@ UObject* UStormSyncArchiveFactory::ImportObject(UClass* InClass, UObject* InOute
 	{
 		if (FileSize == INDEX_NONE)
 		{
-			STORM_SYNC_EDITOR_LOG(Error, TEXT("Can't find file '%s' for import"), *Filename);
+			UE_LOG(LogStormSyncEditor, Error, TEXT("Can't find file '%s' for import"), *Filename);
 		}
 		else
 		{
-			STORM_SYNC_EDITOR_LOG(Display, TEXT("FactoryCreateFile: %s with %s (%i %i %s)"), *InClass->GetName(), *GetClass()->GetName(), bCreateNew, bText, *Filename);
+			UE_LOG(LogStormSyncEditor, Display, TEXT("FactoryCreateFile: %s with %s (%i %i %s)"), *InClass->GetName(), *GetClass()->GetName(), bCreateNew, bText, *Filename);
 			Result = FactoryCreateFile(InClass, InOuter, InName, Flags, *Filename, Params, GWarn, OutCanceled);
 		}
 	}
@@ -83,7 +83,7 @@ UObject* UStormSyncArchiveFactory::ImportObject(UClass* InClass, UObject* InOute
 
 UObject* UStormSyncArchiveFactory::FactoryCreateFile(UClass* InClass, UObject* InParent, FName InName, EObjectFlags Flags, const FString& Filename, const TCHAR* Params, FFeedbackContext* Warn, bool& bOutOperationCanceled)
 {
-	STORM_SYNC_EDITOR_LOG(Display, TEXT("UStormSyncArchiveFactory::FactoryCreateFile InClass: %s, InOuter: %s, InName: %s, Filename: %s"), *GetNameSafe(InClass), *GetNameSafe(InParent), *InName.ToString(), *Filename);
+	UE_LOG(LogStormSyncEditor, Display, TEXT("UStormSyncArchiveFactory::FactoryCreateFile InClass: %s, InOuter: %s, InName: %s, Filename: %s"), *GetNameSafe(InClass), *GetNameSafe(InParent), *InName.ToString(), *Filename);
 
 	// We create the dummy imported object at the dropped location, otherwise UAssetToolsImpl::SyncBrowserToAssets will
 	// select a fallback root folder (i.e. Engine/Content) which is undesired.
@@ -95,7 +95,7 @@ UObject* UStormSyncArchiveFactory::FactoryCreateFile(UClass* InClass, UObject* I
 
 	PendingImports.Add(Filename, FSoftObjectPath(ArchiveData));
 
-	STORM_SYNC_EDITOR_LOG(Display, TEXT("UStormSyncArchiveFactory::FactoryCreateFile for %s"), *GetNameSafe(ArchiveData))
+	UE_LOG(LogStormSyncEditor, Display, TEXT("UStormSyncArchiveFactory::FactoryCreateFile for %s"), *GetNameSafe(ArchiveData));
 	
 	// Broadcast import event and delay import until next tick to avoid blocking the process that files were dragged from
 	// (delay handled internally by Import subsystem)
diff --git a/Engine/Plugins/Experimental/StormSync/Source/StormSyncEditor/Private/Slate/SStormSyncExportWizard.cpp b/Engine/Plugins/Experimental/StormSync/Source/StormSyncEditor/Private/Slate/SStormSyncExportWizard.cpp
--- a/Engine/Plugins/Experimental/StormSync/Source/StormSyncEditor/Private/Slate/SStormSyncExportWizard.cpp
+++ b/Engine/Plugins/Experimental/StormSync/Source/StormSyncEditor/Private/Slate/SStormSyncExportWizard.cpp
@@ -314,7 +314,7 @@ void SStormSyncExportWizard::OnDestinationPageEntered()
 	const TArray<FName> PackageNamesToExport = GetPackageNamesFromReportsData();
 	if (PackageNamesToExport.IsEmpty())
 	{
-		STORM_SYNC_EDITOR_LOG(Error, TEXT("FStormSyncEditorModule::OnDestinationPageEntered - Can't proceed with empty list of files to include"));
+		UE_LOG(LogStormSyncEditor, Error, TEXT("FStormSyncEditorModule::OnDestinationPageEntered - Can't proceed with empty list of files to include"));
 		if (FileDependenciesListView.IsValid())
 		{
 			FileDependenciesListView->RebuildList();
@@ -331,7 +331,7 @@ void SStormSyncExportWizard::OnDestinationPageEntered()
 	TArray<FStormSyncFileDependency> FileDependencies = FStormSyncCoreUtils::GetAvaFileDependenciesFromPackageNames(PackageNamesToExport);
 	for (const FStormSyncFileDependency& FileDependency : FileDependencies)
 	{
-		STORM_SYNC_EDITOR_LOG(Display, TEXT("\tFileDependency: %s"), *FileDependency.ToString());
+		UE_LOG(LogStormSyncEditor, Display, TEXT("\tFileDependency: %s"), *FileDependency.ToString());
 		FileDependencyList.Add(MakeShared<FStormSyncFileDependency>(FileDependency));
 	}
 	
@@ -361,7 +361,7 @@ TArray<FName> SStormSyncExportWizard::GetPackageNamesFromReportsData() const
 	const TArray<FStormSyncReportPackageData>* ReportsDataPtr = ReportPackages.Get();
 	if (!ReportsDataPtr)
 	{
-		STORM_SYNC_EDITOR_LOG(Error, TEXT("SStormSyncExportWizard::GetPackageNamesFromReportsData - Unable to get report data from dialog"));
+		UE_LOG(LogStormSyncEditor, Error, TEXT("SStormSyncExportWizard::GetPackageNamesFromReportsData - Unable to get report data from dialog"));
 		return PackageNames;
 	}
 
diff --git a/Engine/Plugins/Experimental/StormSync/Source/StormSyncEditor/Private/Slate/Status/SStormSyncStatusWidget.cpp b/Engine/Plugins/Experimental/StormSync/Source/StormSyncEditor/Private/Slate/Status/SStormSyncStatusWidget.cpp
--- a/Engine/Plugins/Experimental/StormSync/Source/StormSyncEditor/Private/Slate/Status/SStormSyncStatusWidget.cpp
+++ b/Engine/Plugins/Experimental/StormSync/Source/StormSyncEditor/Private/Slate/Status/SStormSyncStatusWidget.cpp
@@ -19,14 +19,14 @@ BEGIN_SLATE_FUNCTION_BUILD_OPTIMIZATION
 void SStormSyncStatusWidget::Construct(const FArguments& InArgs, const TSharedRef<FStormSyncTransportStatusResponse>& InStatusResponse)
 {
 	FStormSyncTransportStatusResponse& StatusResponse = InStatusResponse.Get();
-	STORM_SYNC_EDITOR_LOG(Display, TEXT("Construct StormSyncStatusWidget with %s"), *StatusResponse.ToString())
+	UE_LOG(LogStormSyncEditor, Display, TEXT("Construct StormSyncStatusWidget with %s"), *StatusResponse.ToString());
 
 	const FStormSyncConnectionInfo RemoteInfo = InStatusResponse->ConnectionInfo;
-	STORM_SYNC_EDITOR_LOG(Display, TEXT("\t Remote: %s"), *RemoteInfo.ToString())
+	UE_LOG(LogStormSyncEditor, Display, TEXT("\t Remote: %s"), *RemoteInfo.ToString());
 	
 	for (const FStormSyncFileModifierInfo& Modifier : StatusResponse.Modifiers)
 	{
-		STORM_SYNC_EDITOR_LOG(Display, TEXT("\t\t Modifier: %s"), *Modifier.ToString())
+		UE_LOG(LogStormSyncEditor, Display, TEXT("\t\t Modifier: %s"), *Modifier.ToString());
 
 		FStormSyncImportFileInfo Info;
 
@@ -290,7 +290,7 @@ FText SStormSyncStatusWidget::GetModifierOperationTooltip(EStormSyncModifierOper
 		case EStormSyncModifierOperation::Overwrite:
 			return LOCTEXT("Operation_Overwrite_Description", "This is an overwrite, meaning both remote and local have the file, but in a different state (mismatch filesize, hash, etc.)");
 		default: 
-			STORM_SYNC_EDITOR_LOG(Warning, TEXT("SStormSyncStatusWidget::GetModifierOperationTooltip: Invalid operation %d"), static_cast<int32>(InModifierOperation));
+			UE_LOG(LogStormSyncEditor, Warning, TEXT("SStormSyncStatusWidget::GetModifierOperationTooltip: Invalid operation %d"), static_cast<int32>(InModifierOperation));
 			return FText::GetEmpty();
 	}
 }
diff --git a/Engine/Plugins/Experimental/StormSync/Source/StormSyncEditor/Private/StormSyncEditor.cpp b/Engine/Plugins/Experimental/StormSync/Source/StormSyncEditor/Private/StormSyncEditor.cpp
--- a/Engine/Plugins/Experimental/StormSync/Source/StormSyncEditor/Private/StormSyncEditor.cpp
+++ b/Engine/Plugins/Experimental/StormSync/Source/StormSyncEditor/Private/StormSyncEditor.cpp
@@ -164,7 +164,7 @@ void FStormSyncEditorModule::RegisterConsoleCommands()
 void FStormSyncEditorModule::ExecuteDumpConnections(const TArray<FString>& Args)
 {
 	FScopeLock ConnectionsLock(&ConnectionsCriticalSection);
-	STORM_SYNC_EDITOR_LOG(Display, TEXT("FStormSyncEditorModule::ExecuteDumpConnections - Listing active connections (%d)"), RegisteredConnections.Num());
+	UE_LOG(LogStormSyncEditor, Display, TEXT("FStormSyncEditorModule::ExecuteDumpConnections - Listing active connections (%d)"), RegisteredConnections.Num());
 	
 	for (const TPair<FMessageAddress, FStormSyncConnectedDevice>& RegisteredConnection : RegisteredConnections)
 	{
@@ -174,7 +174,7 @@ void FStormSyncEditorModule::ExecuteDumpConnections(const TArray<FString>& Args)
 		FText Tooltip = GetEntryTooltipForRemote(MessageAddress, Connection);
 
 		FString ConnectionDescription = FString::Printf(TEXT("\n\n%s\n\n"), *Tooltip.ToString());
-		STORM_SYNC_EDITOR_LOG(Display, TEXT("%s (%s) %s"), *Connection.StormSyncServerAddressId, *Connection.HostName, *ConnectionDescription);
+		UE_LOG(LogStormSyncEditor, Display, TEXT("%s (%s) %s"), *Connection.StormSyncServerAddressId, *Connection.HostName, *ConnectionDescription);
 	}
 }
 
@@ -191,7 +191,7 @@ void FStormSyncEditorModule::UnregisterConsoleCommands()
 void FStormSyncEditorModule::ExecuteRequestStatusCommand(const TArray<FString>& Args) const
 {
 	const FString Argv = FString::Join(Args, TEXT(" "));
-	STORM_SYNC_EDITOR_LOG(Display, TEXT("FStormSyncEditorModule::ExecuteStatusCommand - Argv: %s"), *Argv)
+	UE_LOG(LogStormSyncEditor, Display, TEXT("FStormSyncEditorModule::ExecuteStatusCommand - Argv: %s"), *Argv);
 
 	// Parse command line.
 	TArray<FName> PackageNames;
@@ -199,7 +199,7 @@ void FStormSyncEditorModule::ExecuteRequestStatusCommand(const TArray<FString>&
 
 	if (PackageNames.IsEmpty())
 	{
-		STORM_SYNC_EDITOR_LOG(Error, TEXT("FStormSyncEditorModule::ExecuteStatusCommand - Missing at least one package name to request status."))
+		UE_LOG(LogStormSyncEditor, Error, TEXT("FStormSyncEditorModule::ExecuteStatusCommand - Missing at least one package name to request status."));
 		return;
 	}
 
@@ -207,30 +207,30 @@ void FStormSyncEditorModule::ExecuteRequestStatusCommand(const TArray<FString>&
 	FParse::Value(*Argv, TEXT("-remote="), RemoteAddressId);
 	if (RemoteAddressId.IsEmpty())
 	{
-		STORM_SYNC_EDITOR_LOG(Error, TEXT("FStormSyncEditorModule::ExecuteStatusCommand - Missing -remote parameter"))
+		UE_LOG(LogStormSyncEditor, Error, TEXT("FStormSyncEditorModule::ExecuteStatusCommand - Missing -remote parameter"));
 		return;
 	}
 
 	FMessageAddress RemoteAddress;
 	if (!FMessageAddress::Parse(RemoteAddressId, RemoteAddress))
 	{
-		STORM_SYNC_EDITOR_LOG(Error, TEXT("FStormSyncEditorModule::ExecuteStatusCommand - Unable to parse %s into a Message Address"), *RemoteAddressId);
+		UE_LOG(LogStormSyncEditor, Error, TEXT("FStormSyncEditorModule::ExecuteStatusCommand - Unable to parse %s into a Message Address"), *RemoteAddressId);
 		return;
 	}
 
 	const FOnStormSyncRequestStatusComplete DoneDelegate = FOnStormSyncRequestStatusComplete::CreateLambda([](const TSharedPtr<FStormSyncTransportStatusResponse>& Response)
 	{
-		STORM_SYNC_EDITOR_LOG(Display, TEXT("FStormSyncEditorModule::ExecuteStatusCommand - Received response! %s"), *Response->ToString());
+		UE_LOG(LogStormSyncEditor, Display, TEXT("FStormSyncEditorModule::ExecuteStatusCommand - Received response! %s"), *Response->ToString());
 		SStormSyncStatusWidget::OpenDialog(Response.ToSharedRef());
 	});
 	
-	STORM_SYNC_EDITOR_LOG(Display, TEXT("FStormSyncEditorModule::ExecuteStatusCommand - Send status request on %s"), *RemoteAddress.ToString())
+	UE_LOG(LogStormSyncEditor, Display, TEXT("FStormSyncEditorModule::ExecuteStatusCommand - Send status request on %s"), *RemoteAddress.ToString());
 	IStormSyncTransportClientModule::Get().RequestPackagesStatus(RemoteAddress, PackageNames, DoneDelegate);
 }
 
 void FStormSyncEditorModule::OnServiceDiscoveryConnection(const FString& MessageAddressUID, const FStormSyncConnectedDevice& ConnectedDevice)
 {
-	STORM_SYNC_EDITOR_LOG(Display, TEXT("FStormSyncEditorModule::OnServiceDiscoveryConnection - MessageAddressUID: %s, Hostname: %s, ProjectName: %s"), *MessageAddressUID, *ConnectedDevice.HostName, *ConnectedDevice.ProjectName);
+	UE_LOG(LogStormSyncEditor, Display, TEXT("FStormSyncEditorModule::OnServiceDiscoveryConnection - MessageAddressUID: %s, Hostname: %s, ProjectName: %s"), *MessageAddressUID, *ConnectedDevice.HostName, *ConnectedDevice.ProjectName);
 	
 	FScopeLock ConnectionsLock(&ConnectionsCriticalSection);
 
@@ -243,14 +243,14 @@ void FStormSyncEditorModule::OnServiceDiscoveryConnection(const FString& Message
 
 void FStormSyncEditorModule::OnServiceDiscoveryStateChange(const FString& MessageAddressUID, EStormSyncConnectedDeviceState State)
 {
-	STORM_SYNC_EDITOR_LOG(Display, TEXT("FStormSyncEditorModule::OnServiceDiscoveryStateChange - MessageAddressUID: %s"), *MessageAddressUID);
+	UE_LOG(LogStormSyncEditor, Display, TEXT("FStormSyncEditorModule::OnServiceDiscoveryStateChange - MessageAddressUID: %s"), *MessageAddressUID);
 	
 	FScopeLock ConnectionsLock(&ConnectionsCriticalSection);
 	
 	FMessageAddress RemoteMessageAddress;
 	if (FMessageAddress::Parse(MessageAddressUID, RemoteMessageAddress) && RegisteredConnections.Contains(RemoteMessageAddress))
 	{
-		STORM_SYNC_EDITOR_LOG(Display, TEXT("\tRemote State Changed to %s"), *UEnum::GetValueAsString(State));
+		UE_LOG(LogStormSyncEditor, Display, TEXT("\tRemote State Changed to %s"), *UEnum::GetValueAsString(State));
 		FStormSyncConnectedDevice& Connection = RegisteredConnections.FindChecked(RemoteMessageAddress);
 		Connection.State = State;
 	}
@@ -258,14 +258,14 @@ void FStormSyncEditorModule::OnServiceDiscoveryStateChange(const FString& Messag
 
 void FStormSyncEditorModule::OnServiceDiscoveryServerStatusChange(const FString& MessageAddressUID, bool bIsServerRunning)
 {
-	STORM_SYNC_EDITOR_LOG(Display, TEXT("FStormSyncEditorModule::OnServiceDiscoveryServerStatusChange - MessageAddressUID: %s"), *MessageAddressUID);
+	UE_LOG(LogStormSyncEditor, Display, TEXT("FStormSyncEditorModule::OnServiceDiscoveryServerStatusChange - MessageAddressUID: %s"), *MessageAddressUID);
 	
 	FScopeLock ConnectionsLock(&ConnectionsCriticalSection);
 	
 	FMessageAddress RemoteMessageAddress;
 	if (FMessageAddress::Parse(MessageAddressUID, RemoteMessageAddress) && RegisteredConnections.Contains(RemoteMessageAddress))
 	{
-		STORM_SYNC_EDITOR_LOG(Display, TEXT("\tRemote Server Endpoint State Changed to %s"), bIsServerRunning ? TEXT("running") : TEXT("stopped"));
+		UE_LOG(LogStormSyncEditor, Display, TEXT("\tRemote Server Endpoint State Changed to %s"), bIsServerRunning ? TEXT("running") : TEXT("stopped"));
 		FStormSyncConnectedDevice& Connection = RegisteredConnections.FindChecked(RemoteMessageAddress);
 		Connection.bIsServerRunning = bIsServerRunning;
 	}
@@ -273,7 +273,7 @@ void FStormSyncEditorModule::OnServiceDiscoveryServerStatusChange(const FString&
 
 void FStormSyncEditorModule::OnServiceDiscoveryDisconnection(const FString& MessageAddressUID)
 {
-	STORM_SYNC_EDITOR_LOG(Display, TEXT("FStormSyncEditorModule::OnServiceDiscoveryDisconnection - MessageAddressUID: %s"), *MessageAddressUID);
+	UE_LOG(LogStormSyncEditor, Display, TEXT("FStormSyncEditorModule::OnServiceDiscoveryDisconnection - MessageAddressUID: %s"), *MessageAddressUID);
 	
 	FScopeLock ConnectionsLock(&ConnectionsCriticalSection);
 	
diff --git a/Engine/Plugins/Experimental/StormSync/Source/StormSyncEditor/Private/StormSyncEditorLog.h b/Engine/Plugins/Experimental/StormSync/Source/StormSyncEditor/Private/StormSyncEditorLog.h
--- a/Engine/Plugins/Experimental/StormSync/Source/StormSyncEditor/Private/StormSyncEditorLog.h
+++ b/Engine/Plugins/Experimental/StormSync/Source/StormSyncEditor/Private/StormSyncEditorLog.h
@@ -5,8 +5,3 @@
 #include "Logging/LogMacros.h"
 
 DECLARE_LOG_CATEGORY_EXTERN(LogStormSyncEditor, Display, All);
-
-#define STORM_SYNC_EDITOR_LOG(Verbosity, Format, ...) \
-{ \
-	UE_LOG(LogStormSyncEditor, Verbosity, Format, ##__VA_ARGS__); \
-}
diff --git a/Engine/Plugins/Experimental/StormSync/Source/StormSyncEditor/Private/Subsystems/StormSyncNotificationSubsystem.cpp b/Engine/Plugins/Experimental/StormSync/Source/StormSyncEditor/Private/Subsystems/StormSyncNotificationSubsystem.cpp
--- a/Engine/Plugins/Experimental/StormSync/Source/StormSyncEditor/Private/Subsystems/StormSyncNotificationSubsystem.cpp
+++ b/Engine/Plugins/Experimental/StormSync/Source/StormSyncEditor/Private/Subsystems/StormSyncNotificationSubsystem.cpp
@@ -16,7 +16,7 @@
 void UStormSyncNotificationSubsystem::Initialize(FSubsystemCollectionBase& InCollection)
 {
 	Super::Initialize(InCollection);
-	STORM_SYNC_EDITOR_LOG(Display, TEXT("UStormSyncNotificationSubsystem initialized"))
+	UE_LOG(LogStormSyncEditor, Display, TEXT("UStormSyncNotificationSubsystem initialized"));
 
 	// Create a message log for the notifications to use
 	FMessageLogModule& MessageLogModule = FModuleManager::LoadModuleChecked<FMessageLogModule>("MessageLog");
@@ -39,7 +39,7 @@ void UStormSyncNotificationSubsystem::Initialize(FSubsystemCollectionBase& InCol
 void UStormSyncNotificationSubsystem::Deinitialize()
 {
 	Super::Deinitialize();
-	STORM_SYNC_EDITOR_LOG(Display, TEXT("UStormSyncNotificationSubsystem shutdown"))
+	UE_LOG(LogStormSyncEditor, Display, TEXT("UStormSyncNotificationSubsystem shutdown"));
 
 	if (FModuleManager::Get().IsModuleLoaded("MessageLog"))
 	{
@@ -126,7 +126,7 @@ void UStormSyncNotificationSubsystem::Notify(const EMessageSeverity::Type InSeve
 void UStormSyncNotificationSubsystem::HandlePushResponse(const TSharedPtr<FStormSyncTransportPushResponse>& InResponse)
 {
 	check(InResponse.IsValid())
-	STORM_SYNC_EDITOR_LOG(Display, TEXT("UStormSyncNotificationSubsystem::HandlePushResponse - Response message: %s"), *InResponse->ToString());
+	UE_LOG(LogStormSyncEditor, Display, TEXT("UStormSyncNotificationSubsystem::HandlePushResponse - Response message: %s"), *InResponse->ToString());
 
 	HandleSyncResponse(
 		InResponse,
@@ -138,7 +138,7 @@ void UStormSyncNotificationSubsystem::HandlePushResponse(const TSharedPtr<FStorm
 void UStormSyncNotificationSubsystem::HandlePullResponse(const TSharedPtr<FStormSyncTransportPullResponse>& InResponse)
 {
 	check(InResponse.IsValid())
-	STORM_SYNC_EDITOR_LOG(Display, TEXT("UStormSyncNotificationSubsystem::HandlePullResponse - Response message: %s"), *InResponse->ToString());
+	UE_LOG(LogStormSyncEditor, Display, TEXT("UStormSyncNotificationSubsystem::HandlePullResponse - Response message: %s"), *InResponse->ToString());
 
 	HandleSyncResponse(
 		InResponse,
@@ -150,7 +150,7 @@ void UStormSyncNotificationSubsystem::HandlePullResponse(const TSharedPtr<FStorm
 void UStormSyncNotificationSubsystem::HandleSyncResponse(const TSharedPtr<FStormSyncTransportSyncResponse>& InResponse, const FText& InSuccessTextFormat, const FText& InErrorTextFormat)
 {
 	check(InResponse.IsValid())
-	STORM_SYNC_EDITOR_LOG(Display, TEXT("UStormSyncNotificationSubsystem::HandleSyncResponse - Response message: %s"), *InResponse->ToString());
+	UE_LOG(LogStormSyncEditor, Display, TEXT("UStormSyncNotificationSubsystem::HandleSyncResponse - Response message: %s"), *InResponse->ToString());
 
 	FText NotifyMessage;
 	EMessageSeverity::Type Severity = EMessageSeverity::Info;
@@ -192,19 +192,19 @@ void UStormSyncNotificationSubsystem::HandleSyncResponse(const TSharedPtr<FStorm
 
 void UStormSyncNotificationSubsystem::OnPreStartSendingBuffer(const FString& InRemoteAddress, const FString& InRemoteHostname, const int32 InFileCount)
 {
-	STORM_SYNC_EDITOR_LOG(Display, TEXT("UStormSyncNotificationSubsystem OnPreStartSendingBuffer InFileCount: %d"), InFileCount)
+	UE_LOG(LogStormSyncEditor, Display, TEXT("UStormSyncNotificationSubsystem OnPreStartSendingBuffer InFileCount: %d"), InFileCount);
 	AddProgressBarNotification(InRemoteAddress, InRemoteHostname);
 }
 
 void UStormSyncNotificationSubsystem::OnStartSendingBuffer(const FString& InRemoteAddress, const int32 InBufferSize)
 {
-	STORM_SYNC_EDITOR_LOG(Display, TEXT("UStormSyncNotificationSubsystem OnStartSendingBuffer InBufferSize: %d InRemoteAddress: %s"), InBufferSize, *InRemoteAddress)
+	UE_LOG(LogStormSyncEditor, Display, TEXT("UStormSyncNotificationSubsystem OnStartSendingBuffer InBufferSize: %d InRemoteAddress: %s"), InBufferSize, *InRemoteAddress);
 	UpdateProgressBarNotificationOnStart(InRemoteAddress, InBufferSize);
 }
 
 void UStormSyncNotificationSubsystem::OnReceivingBytes(const FString& InRemoteAddress, const int32 InBytesCount)
 {
-	STORM_SYNC_EDITOR_LOG(Verbose, TEXT("UStormSyncNotificationSubsystem OnReceivingBytes InBytesCount: %d InRemoteAddress: %s"), InBytesCount, *InRemoteAddress)
+	UE_LOG(LogStormSyncEditor, Verbose, TEXT("UStormSyncNotificationSubsystem OnReceivingBytes InBytesCount: %d InRemoteAddress: %s"), InBytesCount, *InRemoteAddress);
 	UpdateProgressBarNotificationOnIncomingBytes(InRemoteAddress, InBytesCount);
 }
 
@@ -212,7 +212,7 @@ void UStormSyncNotificationSubsystem::AddProgressBarNotification(const FString&
 {
 	if (Notifications.Contains(InRemoteAddress))
 	{
-		STORM_SYNC_EDITOR_LOG(Warning, TEXT("Trying to add a notification currently active"))
+		UE_LOG(LogStormSyncEditor, Warning, TEXT("Trying to add a notification currently active"));
 		return;
 	}
 
@@ -243,7 +243,7 @@ void UStormSyncNotificationSubsystem::UpdateProgressBarNotificationOnStart(const
 {
 	if (!Notifications.Contains(InRemoteAddress))
 	{
-		STORM_SYNC_EDITOR_LOG(Warning, TEXT("Trying to update notification for %s, but it's not active"), *InRemoteAddress)
+		UE_LOG(LogStormSyncEditor, Warning, TEXT("Trying to update notification for %s, but it's not active"), *InRemoteAddress);
 		return;
 	}
 
@@ -260,7 +260,7 @@ void UStormSyncNotificationSubsystem::UpdateProgressBarNotificationOnIncomingByt
 {
 	if (!Notifications.Contains(InRemoteAddress))
 	{
-		STORM_SYNC_EDITOR_LOG(Warning, TEXT("Trying to update notification for %s, but it's not active"), *InRemoteAddress)
+		UE_LOG(LogStormSyncEditor, Warning, TEXT("Trying to update notification for %s, but it's not active"), *InRemoteAddress);
 		return;
 	}
 
@@ -274,7 +274,7 @@ void UStormSyncNotificationSubsystem::UpdateProgressBarNotificationOnIncomingByt
 
 		if (Percent >= 1.f && NotificationTask.Notification)
 		{
-			STORM_SYNC_EDITOR_LOG(Display, TEXT("Received 100%% bytes, expire and fadeout notif %f"), Percent);
+			UE_LOG(LogStormSyncEditor, Display, TEXT("Received 100%% bytes, expire and fadeout notif %f"), Percent);
 			NotificationTask.Notification->ExpireAndFadeout();
 			Notifications.Remove(InRemoteAddress);
 		}
@@ -285,13 +285,13 @@ void UStormSyncNotificationSubsystem::OnDismissButtonClicked(const FString InRem
 {
 	if (!Notifications.Contains(InRemoteAddress))
 	{
-		STORM_SYNC_EDITOR_LOG(Warning, TEXT("Trying to dismiss notification for %s, but it's not active"), *InRemoteAddress)
+		UE_LOG(LogStormSyncEditor, Warning, TEXT("Trying to dismiss notification for %s, but it's not active"), *InRemoteAddress);
 		return;
 	}
 
 	const FStormSyncNotificationTask NotificationTask = Notifications.FindChecked(InRemoteAddress);
 
-	STORM_SYNC_EDITOR_LOG(Verbose, TEXT("OnDismissButtonClicked"));
+	UE_LOG(LogStormSyncEditor, Verbose, TEXT("OnDismissButtonClicked"));
 	if (NotificationTask.Notification)
 	{
 		// Expire the notification immediately and ensure it fades quickly so that clicking the buttons feels responsive
diff --git a/Engine/Plugins/Experimental/StormSync/Source/StormSyncImport/Private/StormSyncImportLog.h b/Engine/Plugins/Experimental/StormSync/Source/StormSyncImport/Private/StormSyncImportLog.h
--- a/Engine/Plugins/Experimental/StormSync/Source/StormSyncImport/Private/StormSyncImportLog.h
+++ b/Engine/Plugins/Experimental/StormSync/Source/StormSyncImport/Private/StormSyncImportLog.h
@@ -4,9 +4,4 @@
 
 #include "Logging/LogMacros.h"
 
-STORMSYNCIMPORT_API DECLARE_LOG_CATEGORY_EXTERN(LogStormSyncImport, Display, All);
-
-#define STORM_SYNC_IMPORT_LOG(Verbosity, Format, ...) \
-{ \
-	UE_LOG(LogStormSyncImport, Verbosity, Format, ##__VA_ARGS__); \
-}
+DECLARE_LOG_CATEGORY_EXTERN(LogStormSyncImport, Display, All);
diff --git a/Engine/Plugins/Experimental/StormSync/Source/StormSyncImport/Private/Subsystems/StormSyncImportSubsystem.cpp b/Engine/Plugins/Experimental/StormSync/Source/StormSyncImport/Private/Subsystems/StormSyncImportSubsystem.cpp
--- a/Engine/Plugins/Experimental/StormSync/Source/StormSyncImport/Private/Subsystems/StormSyncImportSubsystem.cpp
+++ b/Engine/Plugins/Experimental/StormSync/Source/StormSyncImport/Private/Subsystems/StormSyncImportSubsystem.cpp
@@ -44,7 +44,7 @@
 
 void UStormSyncImportSubsystem::Initialize(FSubsystemCollectionBase& Collection)
 {
-	STORM_SYNC_IMPORT_LOG(Verbose, TEXT("UStormSyncImportSubsystem::Initialize (World: %s)"), *GetNameSafe(GetWorld()))
+	UE_LOG(LogStormSyncImport, Verbose, TEXT("UStormSyncImportSubsystem::Initialize (World: %s)"), *GetNameSafe(GetWorld()));
 
 #if WITH_EDITOR
 	// Create a message log for the asset tools to use
@@ -57,7 +57,7 @@ void UStormSyncImportSubsystem::Initialize(FSubsystemCollectionBase& Collection)
 
 void UStormSyncImportSubsystem::Deinitialize()
 {
-	STORM_SYNC_IMPORT_LOG(Verbose, TEXT("UStormSyncImportSubsystem::Deinitialize (World: %s)"), *GetNameSafe(GetWorld()))
+	UE_LOG(LogStormSyncImport, Verbose, TEXT("UStormSyncImportSubsystem::Deinitialize (World: %s)"), *GetNameSafe(GetWorld()));
 	
 #if WITH_EDITOR
 	if (FModuleManager::Get().IsModuleLoaded("MessageLog"))
@@ -82,7 +82,7 @@ bool UStormSyncImportSubsystem::EnqueueImportTask(const TSharedPtr<IStormSyncImp
 {
 	if (!InWorld)
 	{
-		STORM_SYNC_IMPORT_LOG(Error, TEXT("UStormSyncImportSubsystem::EnqueueImportTask failed because of invalid world (World: %s)"), *GetNameSafe(InWorld))
+		UE_LOG(LogStormSyncImport, Error, TEXT("UStormSyncImportSubsystem::EnqueueImportTask failed because of invalid world (World: %s)"), *GetNameSafe(InWorld));
 		return false;
 	}
 
@@ -97,7 +97,7 @@ bool UStormSyncImportSubsystem::EnqueueImportTask(const TSharedPtr<IStormSyncImp
 
 	if (!PendingTasks.IsEmpty())
 	{
-		STORM_SYNC_IMPORT_LOG(Error, TEXT("UStormSyncImportSubsystem::EnqueueImportTask failed because of existing pending tasks (World: %s)"), *GetNameSafe(InWorld))
+		UE_LOG(LogStormSyncImport, Error, TEXT("UStormSyncImportSubsystem::EnqueueImportTask failed because of existing pending tasks (World: %s)"), *GetNameSafe(InWorld));
 		return false;
 	}
 
@@ -113,13 +113,13 @@ bool UStormSyncImportSubsystem::EnqueueImportTask(const TSharedPtr<IStormSyncImp
 
 bool UStormSyncImportSubsystem::PerformFileImport(const FString& InFilename)
 {
-	STORM_SYNC_IMPORT_LOG(Display, TEXT("FStormSyncEditorUtils::PerformImport for %s"), *InFilename);
+	UE_LOG(LogStormSyncImport, Display, TEXT("FStormSyncEditorUtils::PerformImport for %s"), *InFilename);
 	TArray<uint8> Data;
 
 	const FStormSyncBufferPtr Buffer = MakeShareable(new TArray<uint8>);
 	if (!FFileHelper::LoadFileToArray(*Buffer.Get(), *InFilename))
 	{
-		STORM_SYNC_IMPORT_LOG(Error, TEXT("Failed to load file '%s' to array"), *InFilename);
+		UE_LOG(LogStormSyncImport, Error, TEXT("Failed to load file '%s' to array"), *InFilename);
 		return false;
 	}
 
@@ -135,7 +135,7 @@ bool UStormSyncImportSubsystem::PerformBufferImport(const FStormSyncPackageDescr
 {
 	if (!InBuffer.IsValid())
 	{
-		STORM_SYNC_IMPORT_LOG(Error, TEXT("UStormSyncImportSubsystem::PerformBufferImport failed to import from invalid buffer"))
+		UE_LOG(LogStormSyncImport, Error, TEXT("UStormSyncImportSubsystem::PerformBufferImport failed to import from invalid buffer"));
 		return false;
 	}
 
@@ -151,7 +151,7 @@ bool UStormSyncImportSubsystem::PerformImport(const FStormSyncPackageDescriptor&
 {
 	if (!InBuffer.IsValid())
 	{
-		STORM_SYNC_IMPORT_LOG(Error, TEXT("UStormSyncImportSubsystem::PerformBufferImport failed to import from invalid buffer"))
+		UE_LOG(LogStormSyncImport, Error, TEXT("UStormSyncImportSubsystem::PerformBufferImport failed to import from invalid buffer"));
 		return false;
 	}
 
@@ -173,7 +173,7 @@ bool UStormSyncImportSubsystem::PerformImport(const FStormSyncPackageDescriptor&
 	ExtractArgs.OnPakPreExtract.BindLambda([&SlowTask, PackageDescriptor, bDryRun](const int32 FileCount)
 	{
 		SlowTask.TotalAmountOfWork = FileCount;
-		STORM_SYNC_IMPORT_LOG(Display, TEXT("\tOnPakPreExtract - File Count: %d"), FileCount)
+		UE_LOG(LogStormSyncImport, Display, TEXT("\tOnPakPreExtract - File Count: %d"), FileCount);
 		
 		if (!bDryRun)
 		{
@@ -185,7 +185,7 @@ bool UStormSyncImportSubsystem::PerformImport(const FStormSyncPackageDescriptor&
 	{
 		if (bDryRun)
 		{
-			STORM_SYNC_IMPORT_LOG(Display, TEXT("\tOnPakPostExtract - File Count: %d"), FileCount)
+			UE_LOG(LogStormSyncImport, Display, TEXT("\tOnPakPostExtract - File Count: %d"), FileCount);
 			return;
 		}
 
@@ -241,14 +241,14 @@ bool UStormSyncImportSubsystem::PerformImport(const FStormSyncPackageDescriptor&
 		
 		if (bDryRun)
 		{
-			STORM_SYNC_IMPORT_LOG(Display, TEXT("\tOnFileExtract"))
-			STORM_SYNC_IMPORT_LOG(Display, TEXT("\t\tPackageName: %s"), *FileDependency.PackageName.ToString())
-			STORM_SYNC_IMPORT_LOG(Display, TEXT("\t\tDestFilepath: %s"), *DestFilepath)
-			STORM_SYNC_IMPORT_LOG(Display, TEXT("\t\tFileSize: %lld"), FileDependency.FileSize)
-			STORM_SYNC_IMPORT_LOG(Display, TEXT("\t\tFileBuffer valid: %s"), FileBuffer.IsValid() ? TEXT("valid") : TEXT("invalid"))
+			UE_LOG(LogStormSyncImport, Display, TEXT("\tOnFileExtract"));
+			UE_LOG(LogStormSyncImport, Display, TEXT("\t\tPackageName: %s"), *FileDependency.PackageName.ToString());
+			UE_LOG(LogStormSyncImport, Display, TEXT("\t\tDestFilepath: %s"), *DestFilepath);
+			UE_LOG(LogStormSyncImport, Display, TEXT("\t\tFileSize: %lld"), FileDependency.FileSize);
+			UE_LOG(LogStormSyncImport, Display, TEXT("\t\tFileBuffer valid: %s"), FileBuffer.IsValid() ? TEXT("valid") : TEXT("invalid"));
 			if (FileBuffer.IsValid())
 			{
-				STORM_SYNC_IMPORT_LOG(Display, TEXT("\t\tFileBuffer size: %d"), FileBuffer->Num())
+				UE_LOG(LogStormSyncImport, Display, TEXT("\t\tFileBuffer size: %d"), FileBuffer->Num());
 			}
 
 			return;
@@ -265,10 +265,10 @@ bool UStormSyncImportSubsystem::PerformImport(const FStormSyncPackageDescriptor&
 	TMap<FString, FString> ExtractedPackages;
 	if (!FStormSyncCoreUtils::ExtractPakBuffer(*InBuffer.Get(), ExtractArgs, ExtractedPackages, Errors))
 	{
-		STORM_SYNC_IMPORT_LOG(Warning, TEXT("FStormSyncEditorUtils::PerformImport - Error extracting package ..."))
+		UE_LOG(LogStormSyncImport, Warning, TEXT("FStormSyncEditorUtils::PerformImport - Error extracting package ..."));
 		for (const FText& Error : Errors)
 		{
-			STORM_SYNC_IMPORT_LOG(Warning, TEXT("\t %s"), *Error.ToString())
+			UE_LOG(LogStormSyncImport, Warning, TEXT("\t %s"), *Error.ToString());
 		}
 
 		return false;
@@ -305,7 +305,7 @@ bool UStormSyncImportSubsystem::FillFilesToImport(const FStormSyncImportFileInfo
 	// File not existing in local project, add it to files to import and early out
 	if (DestFileSize == INDEX_NONE)
 	{
-		STORM_SYNC_IMPORT_LOG(Verbose, TEXT("\t\tFile %s does not exist"), *InDestFilepath)
+		UE_LOG(LogStormSyncImport, Verbose, TEXT("\t\tFile %s does not exist"), *InDestFilepath);
 		FString ShortPath = InDestFilepath;
 		ShortPath.RemoveFromStart(FPaths::ProjectContentDir());
 		FileInfo.ImportReason = FText::Format(LOCTEXT("FileInfo_MissingFile", "Missing file in local project ({0})"), FText::FromString(ShortPath));
@@ -314,10 +314,10 @@ bool UStormSyncImportSubsystem::FillFilesToImport(const FStormSyncImportFileInfo
 		return true;
 	}
 
-	STORM_SYNC_IMPORT_LOG(Verbose, TEXT("\t\tFile %s exist"), *InDestFilepath)
+	UE_LOG(LogStormSyncImport, Verbose, TEXT("\t\tFile %s exist"), *InDestFilepath);
 
 	const bool bSameSize = FileDependency.FileSize == DestFileSize;
-	STORM_SYNC_IMPORT_LOG(Verbose, TEXT("\t\t\tSame Size: %s (%lld vs %lld)"), bSameSize ? TEXT("true") : TEXT("false"), FileDependency.FileSize, DestFileSize)
+	UE_LOG(LogStormSyncImport, Verbose, TEXT("\t\t\tSame Size: %s (%lld vs %lld)"), bSameSize ? TEXT("true") : TEXT("false"), FileDependency.FileSize, DestFileSize);
 
 	// Check file hash from buffer against local file
 	const FMD5Hash ExistingFileMD5 = FMD5Hash::HashFile(*InDestFilepath);
@@ -361,7 +361,7 @@ bool UStormSyncImportSubsystem::FillFilesToImport(const FStormSyncImportFileInfo
 
 void UStormSyncImportSubsystem::HandlePakPreExtract(const FStormSyncPackageDescriptor& InPackageDescriptor, const int32 FileCount)
 {
-	STORM_SYNC_IMPORT_LOG(Display, TEXT("UStormSyncImportSubsystem::HandlePakPreExtract - About to extract %d files for %s"), FileCount, *InPackageDescriptor.ToString());
+	UE_LOG(LogStormSyncImport, Display, TEXT("UStormSyncImportSubsystem::HandlePakPreExtract - About to extract %d files for %s"), FileCount, *InPackageDescriptor.ToString());
 
 	// Reset closed assets to reopen
 	ClosedPackageNames.Empty();
@@ -381,8 +381,8 @@ void UStormSyncImportSubsystem::HandlePakPreExtract(const FStormSyncPackageDescr
 
 void UStormSyncImportSubsystem::HandlePakPostExtract(const FStormSyncPackageDescriptor& InPackageDescriptor, const int32 FileCount) const
 {
-	STORM_SYNC_IMPORT_LOG(Display, TEXT("UStormSyncImportSubsystem::HandlePakPostExtract - Extracted %d files for %s"), FileCount, *InPackageDescriptor.ToString());
-	STORM_SYNC_IMPORT_LOG(Display, TEXT("UStormSyncImportSubsystem::HandlePakPostExtract - Should handle reopening of %d assets"), ClosedPackageNames.Num());
+	UE_LOG(LogStormSyncImport, Display, TEXT("UStormSyncImportSubsystem::HandlePakPostExtract - Extracted %d files for %s"), FileCount, *InPackageDescriptor.ToString());
+	UE_LOG(LogStormSyncImport, Display, TEXT("UStormSyncImportSubsystem::HandlePakPostExtract - Should handle reopening of %d assets"), ClosedPackageNames.Num());
 	
 	if (GetDefault<UStormSyncCoreSettings>()->bEnableHotReloadPackages)
 	{
@@ -426,7 +426,7 @@ void UStormSyncImportSubsystem::HandlePakPostExtract(const FStormSyncPackageDesc
 void UStormSyncImportSubsystem::HandleExistingAssets(TConstArrayView<const FStormSyncImportFileInfo*> InExistingFiles, bool bInShowPrompt)
 {
 #if WITH_EDITOR
-	STORM_SYNC_IMPORT_LOG(Display, TEXT("UStormSyncImportSubsystem::HandleExistingAssets - Package Count: %d"), InExistingFiles.Num());
+	UE_LOG(LogStormSyncImport, Display, TEXT("UStormSyncImportSubsystem::HandleExistingAssets - Package Count: %d"), InExistingFiles.Num());
 
 	TArray<UPackage*> ExistingPackages;
 
@@ -439,12 +439,12 @@ void UStormSyncImportSubsystem::HandleExistingAssets(TConstArrayView<const FStor
 
 		if (UPackage* const Package = LoadPackage(nullptr, *PackageName, LOAD_None))
 		{
-			STORM_SYNC_IMPORT_LOG(Verbose, TEXT("UStormSyncImportSubsystem::HandleExistingAssets - Existing Package found for: %s"), *PackageName);
+			UE_LOG(LogStormSyncImport, Verbose, TEXT("UStormSyncImportSubsystem::HandleExistingAssets - Existing Package found for: %s"), *PackageName);
 			ExistingPackages.Add(Package);
 		}
 		else
 		{
-			STORM_SYNC_IMPORT_LOG(Verbose, TEXT("UStormSyncImportSubsystem::HandleExistingAssets - Package not found for: %s. Not an Existing Asset"), *PackageName);
+			UE_LOG(LogStormSyncImport, Verbose, TEXT("UStormSyncImportSubsystem::HandleExistingAssets - Package not found for: %s. Not an Existing Asset"), *PackageName);
 		}
 	}
 
@@ -475,7 +475,7 @@ void UStormSyncImportSubsystem::HandleExistingAssets(TConstArrayView<const FStor
 
 		if (SourceControlProvider.IsEnabled())
 		{
-			STORM_SYNC_IMPORT_LOG(Log, TEXT("UStormSyncImportSubsystem::HandleExistingAssets - Checking out %d packages with %s")
+			UE_LOG(LogStormSyncImport, Log, TEXT("UStormSyncImportSubsystem::HandleExistingAssets - Checking out %d packages with %s")
 				, ExistingPackages.Num()
 				, *SourceControlProvider.GetName().ToString());
 
@@ -486,7 +486,7 @@ void UStormSyncImportSubsystem::HandleExistingAssets(TConstArrayView<const FStor
 		}
 		else
 		{
-			STORM_SYNC_IMPORT_LOG(Log, TEXT("UStormSyncImportSubsystem::HandleExistingAssets - No Source Control found. Making %d packages writeable")
+			UE_LOG(LogStormSyncImport, Log, TEXT("UStormSyncImportSubsystem::HandleExistingAssets - No Source Control found. Making %d packages writeable")
 				, ExistingPackages.Num());
 
 			bSucceeded = UStormSyncImportSubsystem::MakePackagesWriteable(ExistingPackages) > 0;
@@ -496,7 +496,7 @@ void UStormSyncImportSubsystem::HandleExistingAssets(TConstArrayView<const FStor
 	if (!bSucceeded)
 	{
 		// TODO: at this point in time, the whole import could be cancelled and failed
-		STORM_SYNC_IMPORT_LOG(Error, TEXT("UStormSyncImportSubsystem::HandleExistingAssets - CheckoutPackages failed"));
+		UE_LOG(LogStormSyncImport, Error, TEXT("UStormSyncImportSubsystem::HandleExistingAssets - CheckoutPackages failed"));
 	}
 #endif
 }
@@ -511,19 +511,19 @@ bool UStormSyncImportSubsystem::HandleNewAssets(TConstArrayView<const FStormSync
 
 	if (!ISourceControlModule::Get().IsEnabled())
 	{
-		STORM_SYNC_IMPORT_LOG(Log, TEXT("UStormSyncImportSubsystem::HandleNewAssets - No Source Control found"));
+		UE_LOG(LogStormSyncImport, Log, TEXT("UStormSyncImportSubsystem::HandleNewAssets - No Source Control found"));
 		return false;
 	}
 
 	if (!GetDefault<UEditorLoadingSavingSettings>()->bSCCAutoAddNewFiles)
 	{
-		STORM_SYNC_IMPORT_LOG(Log, TEXT("UStormSyncImportSubsystem::HandleNewAssets - Source Control Auto Add New Files is disabled."));
+		UE_LOG(LogStormSyncImport, Log, TEXT("UStormSyncImportSubsystem::HandleNewAssets - Source Control Auto Add New Files is disabled."));
 		return false;
 	}
 
 	if (!ISourceControlModule::Get().GetProvider().IsAvailable())
 	{
-		STORM_SYNC_IMPORT_LOG(Log, TEXT("UStormSyncImportSubsystem::HandleNewAssets - Source Control Provider is not available."));
+		UE_LOG(LogStormSyncImport, Log, TEXT("UStormSyncImportSubsystem::HandleNewAssets - Source Control Provider is not available."));
 		return false;
 	}
 
@@ -536,7 +536,7 @@ bool UStormSyncImportSubsystem::HandleNewAssets(TConstArrayView<const FStormSync
 			return InFileInfo->DestFilepath;
 		});
 
-	STORM_SYNC_IMPORT_LOG(Log, TEXT("UStormSyncImportSubsystem::HandleNewAssets - Found %d files to auto-add"), FilesToAutoAdd.Num());
+	UE_LOG(LogStormSyncImport, Log, TEXT("UStormSyncImportSubsystem::HandleNewAssets - Found %d files to auto-add"), FilesToAutoAdd.Num());
 
 	const bool bSilent = !bInShowPrompt;
 	return USourceControlHelpers::CheckOutOrAddFiles(FilesToAutoAdd, bSilent);
@@ -546,13 +546,13 @@ bool UStormSyncImportSubsystem::HandleNewAssets(TConstArrayView<const FStormSync
 
 void UStormSyncImportSubsystem::HandlePakAssetExtract(const FStormSyncFileDependency& FileDependency, const FString& DestFilepath, const FStormSyncBufferPtr& FileBuffer)
 {
-	STORM_SYNC_IMPORT_LOG(Display, TEXT("UStormSyncImportSubsystem::HandlePakAssetExtract - Handle extracted package: %s"), *FileDependency.ToString());
-	STORM_SYNC_IMPORT_LOG(Display, TEXT("\tshould extract to %s"), *DestFilepath);
-	STORM_SYNC_IMPORT_LOG(Display, TEXT("\tValid buffer ? %s"), FileBuffer.IsValid() && FileBuffer->Num() != 0 ? TEXT("true") : TEXT("false"));
+	UE_LOG(LogStormSyncImport, Display, TEXT("UStormSyncImportSubsystem::HandlePakAssetExtract - Handle extracted package: %s"), *FileDependency.ToString());
+	UE_LOG(LogStormSyncImport, Display, TEXT("\tshould extract to %s"), *DestFilepath);
+	UE_LOG(LogStormSyncImport, Display, TEXT("\tValid buffer ? %s"), FileBuffer.IsValid() && FileBuffer->Num() != 0 ? TEXT("true") : TEXT("false"));
 
 	if (!FileBuffer.IsValid())
 	{
-		STORM_SYNC_IMPORT_LOG(Error, TEXT("UStormSyncImportSubsystem::HandlePakAssetExtract - Invalid file buffer"));
+		UE_LOG(LogStormSyncImport, Error, TEXT("UStormSyncImportSubsystem::HandlePakAssetExtract - Invalid file buffer"));
 		return;
 	}
 
@@ -662,7 +662,7 @@ void UStormSyncImportSubsystem::CloseEditors(const TArray<FAssetData>& InAssets,
 	for (const FAssetData& AssetData : InAssets)
 	{
 		const UObject* Asset = AssetData.FastGetAsset();
-		STORM_SYNC_IMPORT_LOG(Verbose, TEXT("\tClosing asset: %s (UObject: %s)"), *AssetData.GetFullName(), *GetNameSafe(Asset));
+		UE_LOG(LogStormSyncImport, Verbose, TEXT("\tClosing asset: %s (UObject: %s)"), *AssetData.GetFullName(), *GetNameSafe(Asset));
 
 		if (!Asset)
 		{
@@ -683,14 +683,14 @@ void UStormSyncImportSubsystem::CloseEditors(const TArray<FAssetData>& InAssets,
 
 void UStormSyncImportSubsystem::OpenClosedEditors(const TArray<FString>& ClosedPackageNames)
 {
-	STORM_SYNC_IMPORT_LOG(Verbose, TEXT("UStormSyncImportSubsystem::OpenClosedEditors - ClosedPackageNames: %d"), ClosedPackageNames.Num())
+	UE_LOG(LogStormSyncImport, Verbose, TEXT("UStormSyncImportSubsystem::OpenClosedEditors - ClosedPackageNames: %d"), ClosedPackageNames.Num());
 #if WITH_EDITOR
 	if (!GEditor)
 	{
 		return;
 	}
 	
-	STORM_SYNC_IMPORT_LOG(Verbose, TEXT("UStormSyncImportSubsystem::OpenClosedEditors - ClosedPackageNames: %d (WITH_EDITOR)"), ClosedPackageNames.Num())
+	UE_LOG(LogStormSyncImport, Verbose, TEXT("UStormSyncImportSubsystem::OpenClosedEditors - ClosedPackageNames: %d (WITH_EDITOR)"), ClosedPackageNames.Num());
 	TArray<UObject*> ObjectsToReopen;
 	for (FString ClosedPackageName : ClosedPackageNames)
 	{
@@ -773,7 +773,7 @@ bool UStormSyncImportSubsystem::DeleteAssets(const TArray<FAssetData>& AssetsToD
 	const int32 TotalDeletedObjects = NumPackagesToDelete + NumObjectsToDelete;
 	if (TotalDeletedObjects != AssetsToDelete.Num())
 	{
-		STORM_SYNC_IMPORT_LOG(Warning, TEXT("Failed to delete assets (Deleted %d assets while we were expecting to delete %d assets)"), TotalDeletedObjects, AssetsToDelete.Num());
+		UE_LOG(LogStormSyncImport, Warning, TEXT("Failed to delete assets (Deleted %d assets while we were expecting to delete %d assets)"), TotalDeletedObjects, AssetsToDelete.Num());
 		return false;
 	}
 #endif
@@ -842,7 +842,7 @@ bool UStormSyncImportSubsystem::WriteFile(const FString& DestFilepath, const uin
 		return false;
 	}
 
-	STORM_SYNC_IMPORT_LOG(Verbose, TEXT("UStormSyncImportSubsystem::WriteFile - Creating file `%s`"), *DestFilepath);
+	UE_LOG(LogStormSyncImport, Verbose, TEXT("UStormSyncImportSubsystem::WriteFile - Creating file `%s`"), *DestFilepath);
 
 	// Write to asset
 	AssetHandle->Serialize(FileBuffer, FileSize);
@@ -895,7 +895,7 @@ void UStormSyncImportSubsystem::HotReloadPackages(const TArray<FStormSyncEditorF
 		return;
 	}
 
-	STORM_SYNC_IMPORT_LOG(Display, TEXT("UStormSyncImportSubsystem::HotReloadPackages - Reloading %d packages."), PackageNames.Num());
+	UE_LOG(LogStormSyncImport, Display, TEXT("UStormSyncImportSubsystem::HotReloadPackages - Reloading %d packages."), PackageNames.Num());
 	
 	// Flush loading and clean-up any temporary placeholder packages (due to a package previously being missing on disk)
 	FlushAsyncLoading();
@@ -907,7 +907,7 @@ void UStormSyncImportSubsystem::HotReloadPackages(const TArray<FStormSyncEditorF
 		}
 		if (bRunGC)
 		{
-			STORM_SYNC_IMPORT_LOG(Display, TEXT("UStormSyncImportSubsystem::HotReloadPackages - Garbage Collecting..."));
+			UE_LOG(LogStormSyncImport, Display, TEXT("UStormSyncImportSubsystem::HotReloadPackages - Garbage Collecting..."));
 			CollectGarbage(GARBAGE_COLLECTION_KEEPFLAGS);
 		}
 	}
@@ -943,7 +943,7 @@ void UStormSyncImportSubsystem::HotReloadPackages(const TArray<FStormSyncEditorF
 
 		if (!ErrorMessage.IsEmpty())
 		{
-			STORM_SYNC_IMPORT_LOG(Error, TEXT("UStormSyncImportSubsystem::HotReloadPackages: %s"), *ErrorMessage.ToString());
+			UE_LOG(LogStormSyncImport, Error, TEXT("UStormSyncImportSubsystem::HotReloadPackages: %s"), *ErrorMessage.ToString());
 		}
 	}
 }
diff --git a/Engine/Plugins/Experimental/StormSync/Source/StormSyncImport/Private/Subsystems/StormSyncImportSubsystem.h b/Engine/Plugins/Experimental/StormSync/Source/StormSyncImport/Private/Subsystems/StormSyncImportSubsystem.h
--- a/Engine/Plugins/Experimental/StormSync/Source/StormSyncImport/Private/Subsystems/StormSyncImportSubsystem.h
+++ b/Engine/Plugins/Experimental/StormSync/Source/StormSyncImport/Private/Subsystems/StormSyncImportSubsystem.h
@@ -36,7 +36,7 @@ struct FStormSyncEditorFileReport
  * Contains utility functions and callbacks for hooking into importing.
  */
 UCLASS()
-class STORMSYNCIMPORT_API UStormSyncImportSubsystem final : public UEngineSubsystem
+class UStormSyncImportSubsystem final : public UEngineSubsystem
 {
 	GENERATED_BODY()
 
diff --git a/Engine/Plugins/Experimental/StormSync/Source/StormSyncImport/Private/Subsystems/StormSyncImportWorldSubsystem.cpp b/Engine/Plugins/Experimental/StormSync/Source/StormSyncImport/Private/Subsystems/StormSyncImportWorldSubsystem.cpp
--- a/Engine/Plugins/Experimental/StormSync/Source/StormSyncImport/Private/Subsystems/StormSyncImportWorldSubsystem.cpp
+++ b/Engine/Plugins/Experimental/StormSync/Source/StormSyncImport/Private/Subsystems/StormSyncImportWorldSubsystem.cpp
@@ -13,7 +13,7 @@
 void UStormSyncImportWorldSubsystem::Initialize(FSubsystemCollectionBase& Collection)
 {
 	Super::Initialize(Collection);
-	STORM_SYNC_IMPORT_LOG(Verbose, TEXT("UStormSyncImportWorldSubsystem::Initialize (World: %s) - %s"), *GetNameSafe(GetWorld()), *GetName())
+	UE_LOG(LogStormSyncImport, Verbose, TEXT("UStormSyncImportWorldSubsystem::Initialize (World: %s) - %s"), *GetNameSafe(GetWorld()), *GetName());
 
 	FStormSyncCoreDelegates::OnRequestImportFile.AddUObject(this, &UStormSyncImportWorldSubsystem::HandleImportFile);
 	FStormSyncCoreDelegates::OnRequestImportBuffer.AddUObject(this, &UStormSyncImportWorldSubsystem::HandleImportBuffer);
@@ -22,7 +22,7 @@ void UStormSyncImportWorldSubsystem::Initialize(FSubsystemCollectionBase& Collec
 void UStormSyncImportWorldSubsystem::Deinitialize()
 {
 	Super::Deinitialize();
-	STORM_SYNC_IMPORT_LOG(Verbose, TEXT("UStormSyncImportWorldSubsystem::Deinitialize (World: %s) - %s"), *GetNameSafe(GetWorld()), *GetName())
+	UE_LOG(LogStormSyncImport, Verbose, TEXT("UStormSyncImportWorldSubsystem::Deinitialize (World: %s) - %s"), *GetNameSafe(GetWorld()), *GetName());
 
 	FStormSyncCoreDelegates::OnRequestImportFile.RemoveAll(this);
 	FStormSyncCoreDelegates::OnRequestImportBuffer.RemoveAll(this);
@@ -30,11 +30,11 @@ void UStormSyncImportWorldSubsystem::Deinitialize()
 
 void UStormSyncImportWorldSubsystem::HandleImportFile(const FString& InFilename) const
 {
-	STORM_SYNC_IMPORT_LOG(Verbose, TEXT("UStormSyncImportWorldSubsystem::HandleImportFile InFilename: %s (World: %s)"), *InFilename, *GetNameSafe(GetWorld()))
+	UE_LOG(LogStormSyncImport, Verbose, TEXT("UStormSyncImportWorldSubsystem::HandleImportFile InFilename: %s (World: %s)"), *InFilename, *GetNameSafe(GetWorld()));
 	const UWorld* World = GetWorld();
 	if (!World)
 	{
-		STORM_SYNC_IMPORT_LOG(Error, TEXT("UStormSyncImportWorldSubsystem::ImportNextTick failed because of invalid world (World: %s)"), *GetNameSafe(GetWorld()))
+		UE_LOG(LogStormSyncImport, Error, TEXT("UStormSyncImportWorldSubsystem::ImportNextTick failed because of invalid world (World: %s)"), *GetNameSafe(GetWorld()));
 		return;
 	}
 
@@ -45,22 +45,22 @@ void UStormSyncImportWorldSubsystem::HandleImportBuffer(const FStormSyncPackageD
 {
 	if (!InBuffer.IsValid())
 	{
-		STORM_SYNC_IMPORT_LOG(Error, TEXT("UStormSyncImportSubsystem::HandleImportBuffer failed from invalid buffer"))
+		UE_LOG(LogStormSyncImport, Error, TEXT("UStormSyncImportSubsystem::HandleImportBuffer failed from invalid buffer"));
 		return;
 	}
 	
-	STORM_SYNC_IMPORT_LOG(
+	UE_LOG(LogStormSyncImport, 
 		Verbose,
 		TEXT("UStormSyncImportWorldSubsystem::HandleImportBuffer InPackageDescriptor: %s, BufferSize: %d (World: %s)"),
 		*InPackageDescriptor.ToString(),
 		InBuffer->Num(),
 		*GetNameSafe(GetWorld())
-	)
+	);
 
 	const UWorld* World = GetWorld();
 	if (!World)
 	{
-		STORM_SYNC_IMPORT_LOG(Error, TEXT("UStormSyncImportWorldSubsystem::HandleImportBuffer failed because of invalid world (World: %s)"), *GetNameSafe(GetWorld()))
+		UE_LOG(LogStormSyncImport, Error, TEXT("UStormSyncImportWorldSubsystem::HandleImportBuffer failed because of invalid world (World: %s)"), *GetNameSafe(GetWorld()));
 		return;
 	}
 
diff --git a/Engine/Plugins/Experimental/StormSync/Source/StormSyncImport/Private/Subsystems/StormSyncImportWorldSubsystem.h b/Engine/Plugins/Experimental/StormSync/Source/StormSyncImport/Private/Subsystems/StormSyncImportWorldSubsystem.h
--- a/Engine/Plugins/Experimental/StormSync/Source/StormSyncImport/Private/Subsystems/StormSyncImportWorldSubsystem.h
+++ b/Engine/Plugins/Experimental/StormSync/Source/StormSyncImport/Private/Subsystems/StormSyncImportWorldSubsystem.h
@@ -17,7 +17,7 @@
  */
 
 UCLASS()
-class STORMSYNCIMPORT_API UStormSyncImportWorldSubsystem : public UWorldSubsystem
+class UStormSyncImportWorldSubsystem : public UWorldSubsystem
 {
 	GENERATED_BODY()
 
diff --git a/Engine/Plugins/Experimental/StormSync/Source/StormSyncImport/Private/Tasks/StormSyncImportBufferTask.cpp b/Engine/Plugins/Experimental/StormSync/Source/StormSyncImport/Private/Tasks/StormSyncImportBufferTask.cpp
--- a/Engine/Plugins/Experimental/StormSync/Source/StormSyncImport/Private/Tasks/StormSyncImportBufferTask.cpp
+++ b/Engine/Plugins/Experimental/StormSync/Source/StormSyncImport/Private/Tasks/StormSyncImportBufferTask.cpp
@@ -9,10 +9,10 @@ void FStormSyncImportBufferTask::Run()
 {
 	if (!Buffer.IsValid())
 	{
-		STORM_SYNC_IMPORT_LOG(Error, TEXT("FStormSyncImportBufferTask::Run failed on invalid buffer"))
+		UE_LOG(LogStormSyncImport, Error, TEXT("FStormSyncImportBufferTask::Run failed on invalid buffer"));
 		return;
 	}
 
-	STORM_SYNC_IMPORT_LOG(Display, TEXT("FStormSyncImportBufferTask::Run for buffer of size %d"), Buffer->Num())
+	UE_LOG(LogStormSyncImport, Display, TEXT("FStormSyncImportBufferTask::Run for buffer of size %d"), Buffer->Num());
 	UStormSyncImportSubsystem::Get().PerformBufferImport(PackageDescriptor, MoveTemp(Buffer));
 }
diff --git a/Engine/Plugins/Experimental/StormSync/Source/StormSyncImport/Private/Tasks/StormSyncImportFilesTask.cpp b/Engine/Plugins/Experimental/StormSync/Source/StormSyncImport/Private/Tasks/StormSyncImportFilesTask.cpp
--- a/Engine/Plugins/Experimental/StormSync/Source/StormSyncImport/Private/Tasks/StormSyncImportFilesTask.cpp
+++ b/Engine/Plugins/Experimental/StormSync/Source/StormSyncImport/Private/Tasks/StormSyncImportFilesTask.cpp
@@ -8,7 +8,7 @@
 
 void FStormSyncImportFilesTask::Run()
 {
-	STORM_SYNC_IMPORT_LOG(Display, TEXT("FStormSyncImportFilesTask::Run for %s"), *Filename)
+	UE_LOG(LogStormSyncImport, Display, TEXT("FStormSyncImportFilesTask::Run for %s"), *Filename);
 	UStormSyncImportSubsystem::Get().PerformFileImport(Filename);
 	FStormSyncCoreDelegates::OnFileImported.Broadcast(Filename);
 }
diff --git a/Engine/Plugins/Experimental/StormSync/Source/StormSyncTransportClient/Private/Services/StormSyncPullMessageService.cpp b/Engine/Plugins/Experimental/StormSync/Source/StormSyncTransportClient/Private/Services/StormSyncPullMessageService.cpp
--- a/Engine/Plugins/Experimental/StormSync/Source/StormSyncTransportClient/Private/Services/StormSyncPullMessageService.cpp
+++ b/Engine/Plugins/Experimental/StormSync/Source/StormSyncTransportClient/Private/Services/StormSyncPullMessageService.cpp
@@ -55,7 +55,7 @@ void FStormSyncPullMessageService::RequestPullPackages(const FMessageAddress& In
 		if (FileDependencies.IsEmpty())
 		{
 			const FText ErrorText = LOCTEXT("FileDependencies_Error", "Error before pulling: File dependencies array is empty, something went wrong");
-			STORM_SYNC_CLIENT_LOG(Error, TEXT("FStormSyncPullMessageService::RequestPullPackages - %s"), *ErrorText.ToString())
+			UE_LOG(LogStormSyncClient, Error, TEXT("FStormSyncPullMessageService::RequestPullPackages - %s"), *ErrorText.ToString());
 			
 			FStormSyncTransportPullResponse Payload = LocalThis->CreatePullResponsePayload(PackageNames, PackageDescriptor);
 			Payload.Status = EStormSyncResponseResult::Error;
@@ -68,7 +68,7 @@ void FStormSyncPullMessageService::RequestPullPackages(const FMessageAddress& In
 		if (!MessageEndpoint.IsValid())
 		{
 			const FText ErrorText = LOCTEXT("InvalidMessageEndpoint_Error", "Error before pulling: Unable to get message endpoint");
-			STORM_SYNC_CLIENT_LOG(Error, TEXT("FStormSyncPullMessageService::RequestPullPackages - %s"), *ErrorText.ToString())
+			UE_LOG(LogStormSyncClient, Error, TEXT("FStormSyncPullMessageService::RequestPullPackages - %s"), *ErrorText.ToString());
 
 			FStormSyncTransportPullResponse Payload = LocalThis->CreatePullResponsePayload(PackageNames, PackageDescriptor);
 			Payload.Status = EStormSyncResponseResult::Error;
@@ -82,7 +82,7 @@ void FStormSyncPullMessageService::RequestPullPackages(const FMessageAddress& In
 		if (!PullRequestMessage.IsValid())
 		{
 			const FText ErrorText = LOCTEXT("InvalidPullRequestMessage_Error", "Error before pulling: Pull request message is invalid");
-			STORM_SYNC_CLIENT_LOG(Error, TEXT("FStormSyncPullMessageService::RequestPullPackages - %s"), *ErrorText.ToString())
+			UE_LOG(LogStormSyncClient, Error, TEXT("FStormSyncPullMessageService::RequestPullPackages - %s"), *ErrorText.ToString());
 
 			FStormSyncTransportPullResponse Payload = LocalThis->CreatePullResponsePayload(PackageNames, PackageDescriptor);
 			Payload.Status = EStormSyncResponseResult::Error;
@@ -102,9 +102,9 @@ void FStormSyncPullMessageService::RequestPullPackages(const FMessageAddress& In
 		PullRequestMessage->HostAddress = FStormSyncTransportNetworkUtils::GetTcpEndpointAddress();
 		PullRequestMessage->HostAdapterAddresses = FStormSyncTransportNetworkUtils::GetLocalAdapterAddresses();
 
-		STORM_SYNC_CLIENT_LOG(Display, TEXT("FStormSyncPullMessageService::RequestPullPackages - FileDependencies: %d"), FileDependencies.Num())
-		STORM_SYNC_CLIENT_LOG(Display, TEXT("FStormSyncPullMessageService::RequestPullPackages - Message: %s"), *PullRequestMessage->ToString())
-		STORM_SYNC_CLIENT_LOG(Display, TEXT("FStormSyncPullMessageService::RequestPullPackages - Syncing package descriptor %s"), *PullRequestMessage->PackageDescriptor.ToString())
+		UE_LOG(LogStormSyncClient, Display, TEXT("FStormSyncPullMessageService::RequestPullPackages - FileDependencies: %d"), FileDependencies.Num());
+		UE_LOG(LogStormSyncClient, Display, TEXT("FStormSyncPullMessageService::RequestPullPackages - Message: %s"), *PullRequestMessage->ToString());
+		UE_LOG(LogStormSyncClient, Display, TEXT("FStormSyncPullMessageService::RequestPullPackages - Syncing package descriptor %s"), *PullRequestMessage->PackageDescriptor.ToString());
 
 		MessageEndpoint->Send(PullRequestMessage.Release(), Recipient);
 	});
@@ -124,7 +124,7 @@ void FStormSyncPullMessageService::AddResponseHandler(const FGuid& InId, const F
 
 void FStormSyncPullMessageService::HandlePullRequestMessage(const FStormSyncTransportPullRequest& InMessage, const TSharedRef<IMessageContext, ESPMode::ThreadSafe>& InMessageContext)
 {
-	STORM_SYNC_CLIENT_LOG(Display, TEXT("FStormSyncPullMessageService::HandlePullRequestMessage - Received pull request Message from %s (%s)"), *InMessageContext->GetSender().ToString(), *InMessage.ToString())
+	UE_LOG(LogStormSyncClient, Display, TEXT("FStormSyncPullMessageService::HandlePullRequestMessage - Received pull request Message from %s (%s)"), *InMessageContext->GetSender().ToString(), *InMessage.ToString());
 
 	// We received a pull request message from a remote, which contains the full list of "file dependencies" with file state (hash, timestamp), the remote has and wants an update
 	// We have to compare this list with the state of these files locally, and send back a buffer with just the files in a mismatched state
@@ -157,25 +157,25 @@ void FStormSyncPullMessageService::HandlePullRequestMessage(const FStormSyncTran
 		const TSharedPtr<FStormSyncTransportClientEndpoint> LocalEndpoint = LocalThis->GetLocalEndpoint();
 		if (!LocalEndpoint.IsValid())
 		{
-			STORM_SYNC_CLIENT_LOG(Error, TEXT("FStormSyncPullMessageService::HandlePullRequestMessage - Invalid local endpoint."))
+			UE_LOG(LogStormSyncClient, Error, TEXT("FStormSyncPullMessageService::HandlePullRequestMessage - Invalid local endpoint."));
 			return;
 		}
 
 		// Prepare response message payload
 		FStormSyncTransportPullResponse Response = LocalThis->CreatePullResponsePayload(IncomingMessage, true);
 
-		STORM_SYNC_CLIENT_LOG(
+		UE_LOG(LogStormSyncClient, 
 			Display,
 			TEXT("FStormSyncPullMessageService::HandlePullRequestMessage - Initializing tcp connection for %s and message: %s"),
 			*RemoteAddress.ToString(),
 			*Response.ToString()
-		)
+		);
 
 		// Case of no modifiers (both remote in sync), early out and respond with payload
 		if (Response.Modifiers.IsEmpty())
 		{
 			const FText ErrorText = LOCTEXT("State_In_Sync", "Detected empty list of modifiers, meaning assets are in the same state.");
-			STORM_SYNC_CLIENT_LOG(Display, TEXT("\t FStormSyncPullMessageService::HandlePullRequestMessage - %s"), *ErrorText.ToString())
+			UE_LOG(LogStormSyncClient, Display, TEXT("\t FStormSyncPullMessageService::HandlePullRequestMessage - %s"), *ErrorText.ToString());
 
 			// No modifiers found, meaning both ends are in sync
 			Response.Status = EStormSyncResponseResult::Success;
@@ -196,7 +196,7 @@ void FStormSyncPullMessageService::HandlePullRequestMessage(const FStormSyncTran
 			// If text is non empty, indicates an error happened
 			if (!Payload->bSuccess)
 			{
-				STORM_SYNC_CLIENT_LOG(Error, TEXT("FStormSyncPullMessageService::HandlePullRequestMessage - %s"), *Payload->ErrorText.ToString())
+				UE_LOG(LogStormSyncClient, Error, TEXT("FStormSyncPullMessageService::HandlePullRequestMessage - %s"), *Payload->ErrorText.ToString());
 
 				// Got an error, send it back to remote requester
 				FStormSyncTransportPullResponse ErrorResponse = LocalThis->CreatePullResponsePayload(IncomingMessage);
@@ -210,7 +210,7 @@ void FStormSyncPullMessageService::HandlePullRequestMessage(const FStormSyncTran
 		if (!Connection.IsValid())
 		{
 			const FText ErrorText = LOCTEXT("Error_Active_Connection", "Returned active connection is invalid from endpoint StartSendingBuffer()");
-			STORM_SYNC_CLIENT_LOG(Error, TEXT("FStormSyncPullMessageService::HandlePullRequestMessage - %s"), *ErrorText.ToString())
+			UE_LOG(LogStormSyncClient, Error, TEXT("FStormSyncPullMessageService::HandlePullRequestMessage - %s"), *ErrorText.ToString());
 
 			// Something went wrong with the tcp connection, send it back as an error to remote requester
 			FStormSyncTransportPullResponse ErrorResponse = LocalThis->CreatePullResponsePayload(IncomingMessage);
@@ -227,7 +227,7 @@ void FStormSyncPullMessageService::HandlePullRequestMessage(const FStormSyncTran
 
 void FStormSyncPullMessageService::HandlePullResponseMessage(const FStormSyncTransportPullResponse& InMessage, const TSharedRef<IMessageContext, ESPMode::ThreadSafe>& InMessageContext)
 {
-	STORM_SYNC_CLIENT_LOG(Display, TEXT("FStormSyncPullMessageService::HandlePullResponseMessage - Sender: %s, Message: %s"), *InMessageContext->GetSender().ToString(), *InMessage.ToString())
+	UE_LOG(LogStormSyncClient, Display, TEXT("FStormSyncPullMessageService::HandlePullResponseMessage - Sender: %s, Message: %s"), *InMessageContext->GetSender().ToString(), *InMessage.ToString());
 
 	FScopeLock ScopeLock(&HandlersCriticalSection);
 	const FOnStormSyncPullComplete* ResponseHandler = PendingResponseHandlers.Find(InMessage.MessageId);
@@ -244,33 +244,33 @@ void FStormSyncPullMessageService::HandleTransferComplete(const FStormSyncTransp
 	Response.Status = EStormSyncResponseResult::Success;
 	Response.StatusText = FText::Format(LOCTEXT("Pull_Success_TransferComplete", "Successfully pulled {0} package names"), FText::AsNumber(InResponseMessage.PackageNames.Num()));
 	
-	STORM_SYNC_CLIENT_LOG(Display, TEXT("FStormSyncPullMessageService::HandleTransferComplete - TCP transfer done for message with ID: %s (Remote: %s)"), *Response.MessageId.ToString(), *InRemoteAddress.ToString())
+	UE_LOG(LogStormSyncClient, Display, TEXT("FStormSyncPullMessageService::HandleTransferComplete - TCP transfer done for message with ID: %s (Remote: %s)"), *Response.MessageId.ToString(), *InRemoteAddress.ToString());
 	SendPullResponse(Response, InRemoteAddress);
 }
 
 void FStormSyncPullMessageService::SendPullResponse(const FStormSyncTransportPullResponse& InResponsePayload, const FMessageAddress& InRemoteAddress) const
 {
-	STORM_SYNC_CLIENT_LOG(Display, TEXT("FStormSyncPullMessageService::SendPullResponse - Sending to remote %s payload: %s"), *InRemoteAddress.ToString(), *InResponsePayload.ToString())
+	UE_LOG(LogStormSyncClient, Display, TEXT("FStormSyncPullMessageService::SendPullResponse - Sending to remote %s payload: %s"), *InRemoteAddress.ToString(), *InResponsePayload.ToString());
 
 	const TSharedPtr<FMessageEndpoint, ESPMode::ThreadSafe> MessageEndpoint = GetMessageEndpoint();
 	if (!MessageEndpoint.IsValid())
 	{
-		STORM_SYNC_CLIENT_LOG(Error, TEXT("FStormSyncPullMessageService::SendPullResponse - Unable to get client message endpoint (trying sending to remote: %s)"), *InRemoteAddress.ToString())
+		UE_LOG(LogStormSyncClient, Error, TEXT("FStormSyncPullMessageService::SendPullResponse - Unable to get client message endpoint (trying sending to remote: %s)"), *InRemoteAddress.ToString());
 		return;
 	}
 
-	STORM_SYNC_CLIENT_LOG(
+	UE_LOG(LogStormSyncClient, 
 		Display,
 		TEXT("FStormSyncPullMessageService::SendPullResponse - Sending Pull Response with ID: %s (Remote: %s) - Message: %s"),
 		*InResponsePayload.MessageId.ToString(),
 		*InRemoteAddress.ToString(),
 		*InResponsePayload.ToString()
-	)
+	);
 
 	TUniquePtr<FStormSyncTransportPullResponse> Message(FMessageEndpoint::MakeMessage<FStormSyncTransportPullResponse>(InResponsePayload));
 	if (!Message.IsValid())
 	{
-		STORM_SYNC_CLIENT_LOG(Error, TEXT("FStormSyncPullMessageService::SendPullResponse - Pull response message is invalid"))
+		UE_LOG(LogStormSyncClient, Error, TEXT("FStormSyncPullMessageService::SendPullResponse - Pull response message is invalid"));
 		return;
 	}
 
@@ -279,9 +279,9 @@ void FStormSyncPullMessageService::SendPullResponse(const FStormSyncTransportPul
 
 bool FStormSyncPullMessageService::ValidateIncomingPullMessage(const FStormSyncTransportPullRequest& InMessage, FText& OutErrorText) const
 {
-	STORM_SYNC_CLIENT_LOG(Display, TEXT("FStormSyncPullMessageService::ValidateIncomingPullMessage - Validate incoming pull message: %s"), *InMessage.ToString())
+	UE_LOG(LogStormSyncClient, Display, TEXT("FStormSyncPullMessageService::ValidateIncomingPullMessage - Validate incoming pull message: %s"), *InMessage.ToString());
 	
-	STORM_SYNC_CLIENT_LOG(Display, TEXT("FStormSyncPullMessageService::ValidateIncomingPullMessage - Checking for missing package names on disk ... (%s)"), *InMessage.ToString())
+	UE_LOG(LogStormSyncClient, Display, TEXT("FStormSyncPullMessageService::ValidateIncomingPullMessage - Checking for missing package names on disk ... (%s)"), *InMessage.ToString());
 
 	TArray<FName> PackageNames = InMessage.PackageNames;
 	TArray<FString> MissingPackageNames;
diff --git a/Engine/Plugins/Experimental/StormSync/Source/StormSyncTransportClient/Private/Services/StormSyncPushMessageService.cpp b/Engine/Plugins/Experimental/StormSync/Source/StormSyncTransportClient/Private/Services/StormSyncPushMessageService.cpp
--- a/Engine/Plugins/Experimental/StormSync/Source/StormSyncTransportClient/Private/Services/StormSyncPushMessageService.cpp
+++ b/Engine/Plugins/Experimental/StormSync/Source/StormSyncTransportClient/Private/Services/StormSyncPushMessageService.cpp
@@ -52,7 +52,7 @@ void FStormSyncPushMessageService::RequestPushPackages(const FMessageAddress& In
 			if (FileDependencies.IsEmpty())
 			{
 				const FText ErrorText = LOCTEXT("FileDependencies_Error", "Error before pushing: File dependencies array is empty, something went wrong");
-				STORM_SYNC_CLIENT_LOG(Error, TEXT("FStormSyncPushMessageService::RequestPushPackages - %s"), *ErrorText.ToString())
+				UE_LOG(LogStormSyncClient, Error, TEXT("FStormSyncPushMessageService::RequestPushPackages - %s"), *ErrorText.ToString());
 				
 				FStormSyncTransportPushResponse Payload = LocalThis->CreatePushResponsePayload(PackageNames, PackageDescriptor);
 				Payload.Status = EStormSyncResponseResult::Error;
@@ -66,7 +66,7 @@ void FStormSyncPushMessageService::RequestPushPackages(const FMessageAddress& In
 			if (!MessageEndpoint.IsValid())
 			{
 				const FText ErrorText = LOCTEXT("InvalidMessageEndpoint_Error", "Error before pushing: Unable to get message endpoint");
-				STORM_SYNC_CLIENT_LOG(Error, TEXT("FStormSyncPushMessageService::RequestPushPackages - %s"), *ErrorText.ToString())
+				UE_LOG(LogStormSyncClient, Error, TEXT("FStormSyncPushMessageService::RequestPushPackages - %s"), *ErrorText.ToString());
 
 				FStormSyncTransportPushResponse Payload = LocalThis->CreatePushResponsePayload(PackageNames, PackageDescriptor);
 				Payload.Status = EStormSyncResponseResult::Error;
@@ -80,7 +80,7 @@ void FStormSyncPushMessageService::RequestPushPackages(const FMessageAddress& In
 			if (!PushRequestMessage.IsValid())
 			{
 				const FText ErrorText = LOCTEXT("InvalidPullRequestMessage_Error", "Error before pushing: Push request message is invalid");
-				STORM_SYNC_CLIENT_LOG(Error, TEXT("FStormSyncPushMessageService::RequestPushPackages - %s"), *ErrorText.ToString())
+				UE_LOG(LogStormSyncClient, Error, TEXT("FStormSyncPushMessageService::RequestPushPackages - %s"), *ErrorText.ToString());
 
 				FStormSyncTransportPushResponse Payload = LocalThis->CreatePushResponsePayload(PackageNames, PackageDescriptor);
 				Payload.Status = EStormSyncResponseResult::Error;
@@ -94,9 +94,9 @@ void FStormSyncPushMessageService::RequestPushPackages(const FMessageAddress& In
 
 			PushRequestMessage->PackageDescriptor.Dependencies = FileDependencies;
 
-			STORM_SYNC_CLIENT_LOG(Display, TEXT("FStormSyncPushMessageService::RequestPushPackages - FileDependencies: %d"), FileDependencies.Num())
-			STORM_SYNC_CLIENT_LOG(Display, TEXT("FStormSyncPushMessageService::RequestPushPackages - Message: %s"), *PushRequestMessage->ToString())
-			STORM_SYNC_CLIENT_LOG(Display, TEXT("FStormSyncPushMessageService::RequestPushPackages - Syncing package descriptor %s"), *PushRequestMessage->PackageDescriptor.ToString())
+			UE_LOG(LogStormSyncClient, Display, TEXT("FStormSyncPushMessageService::RequestPushPackages - FileDependencies: %d"), FileDependencies.Num());
+			UE_LOG(LogStormSyncClient, Display, TEXT("FStormSyncPushMessageService::RequestPushPackages - Message: %s"), *PushRequestMessage->ToString());
+			UE_LOG(LogStormSyncClient, Display, TEXT("FStormSyncPushMessageService::RequestPushPackages - Syncing package descriptor %s"), *PushRequestMessage->PackageDescriptor.ToString());
 
 			MessageEndpoint->Send(PushRequestMessage.Release(), Recipient);
 		});
@@ -116,7 +116,7 @@ void FStormSyncPushMessageService::AddResponseHandler(const FGuid& InId, const F
 
 void FStormSyncPushMessageService::HandlePushRequestMessage(const FStormSyncTransportPushRequest& InMessage, const TSharedRef<IMessageContext, ESPMode::ThreadSafe>& InMessageContext)
 {
-	STORM_SYNC_CLIENT_LOG(Display, TEXT("FStormSyncPushMessageService::HandlePushRequestMessage - Received push request Message from %s (%s)"), *InMessageContext->GetSender().ToString(), *InMessage.ToString())
+	UE_LOG(LogStormSyncClient, Display, TEXT("FStormSyncPushMessageService::HandlePushRequestMessage - Received push request Message from %s (%s)"), *InMessageContext->GetSender().ToString(), *InMessage.ToString());
 
 	FGuid MessageId = InMessage.MessageId;
 	TArray<FName> LocalPackageNames = InMessage.PackageNames;
@@ -130,7 +130,7 @@ void FStormSyncPushMessageService::HandlePushRequestMessage(const FStormSyncTran
 	{
 		if (!MessageEndpointPtr.IsValid())
 		{
-			STORM_SYNC_CLIENT_LOG(Error, TEXT("FStormSyncPushMessageService::HandlePushRequestMessage - Invalid MessageEndpoint, can't sync response back"))
+			UE_LOG(LogStormSyncClient, Error, TEXT("FStormSyncPushMessageService::HandlePushRequestMessage - Invalid MessageEndpoint, can't sync response back"));
 			return;
 		}
 
@@ -138,7 +138,7 @@ void FStormSyncPushMessageService::HandlePushRequestMessage(const FStormSyncTran
 		TUniquePtr<FStormSyncTransportPushResponse> Message(FMessageEndpoint::MakeMessage<FStormSyncTransportPushResponse>(MessageId));
 		if (!Message.IsValid())
 		{
-			STORM_SYNC_CLIENT_LOG(Error, TEXT("FStormSyncStatusMessageService::RequestStatus - Sync response message is invalid"))
+			UE_LOG(LogStormSyncClient, Error, TEXT("FStormSyncStatusMessageService::RequestStatus - Sync response message is invalid"));
 			return;
 		}
 
@@ -153,10 +153,10 @@ void FStormSyncPushMessageService::HandlePushRequestMessage(const FStormSyncTran
 		Message->HostAddress = FStormSyncTransportNetworkUtils::GetTcpEndpointAddress();
 		Message->HostAdapterAddresses = FStormSyncTransportNetworkUtils::GetLocalAdapterAddresses();
 
-		STORM_SYNC_CLIENT_LOG(Display, TEXT("FStormSyncPushMessageService::HandlePushRequestMessage - Sending back to %s sync response: %s"), *RemoteAddress.ToString(), *Message->ToString())
+		UE_LOG(LogStormSyncClient, Display, TEXT("FStormSyncPushMessageService::HandlePushRequestMessage - Sending back to %s sync response: %s"), *RemoteAddress.ToString(), *Message->ToString());
 		if (FileModifiers.IsEmpty())
 		{
-			STORM_SYNC_CLIENT_LOG(Display, TEXT("\tSending back empty list of modifiers, meaning pak are synced."))
+			UE_LOG(LogStormSyncClient, Display, TEXT("\tSending back empty list of modifiers, meaning pak are synced."));
 		}
 
 		MessageEndpointPtr->Send(Message.Release(), RemoteAddress);
@@ -167,15 +167,15 @@ void FStormSyncPushMessageService::HandlePushResponseMessage(const FStormSyncTra
 {
 	const FMessageAddress Sender = InMessageContext->GetSender();
 
-	STORM_SYNC_CLIENT_LOG(Display, TEXT("FStormSyncPushMessageService::HandlePushResponseMessage - Received sync response from %s - Message: %s"), *Sender.ToString(), *InMessage.ToString())
+	UE_LOG(LogStormSyncClient, Display, TEXT("FStormSyncPushMessageService::HandlePushResponseMessage - Received sync response from %s - Message: %s"), *Sender.ToString(), *InMessage.ToString());
 
 	if (InMessage.Modifiers.IsEmpty())
 	{
 		// Remote sent back empty list of modifiers, meaning assets are synced. Do nothing.
-		STORM_SYNC_CLIENT_LOG(Display, TEXT("FStormSyncPushMessageService::HandlePushResponseMessage - Remote sent back empty list of modifiers, meaning pak are synced. Do nothing."))
+		UE_LOG(LogStormSyncClient, Display, TEXT("FStormSyncPushMessageService::HandlePushResponseMessage - Remote sent back empty list of modifiers, meaning pak are synced. Do nothing."));
 		
 		const FText ErrorText = LOCTEXT("State_In_Sync", "Detected empty list of modifiers, meaning assets are in the same state.");
-		STORM_SYNC_CLIENT_LOG(Display, TEXT("\t FStormSyncPushMessageService::HandlePushResponseMessage - %s"), *ErrorText.ToString())
+		UE_LOG(LogStormSyncClient, Display, TEXT("\t FStormSyncPushMessageService::HandlePushResponseMessage - %s"), *ErrorText.ToString());
 		
 		// But notify caller about completion
 		const TSharedPtr<FStormSyncTransportPushResponse> ErrorResponse = MakeShared<FStormSyncTransportPushResponse>(InMessage);
@@ -191,7 +191,7 @@ void FStormSyncPushMessageService::HandlePushResponseMessage(const FStormSyncTra
 	const TSharedPtr<FStormSyncTransportClientEndpoint> LocalEndpoint = GetLocalEndpoint();
 	if (!LocalEndpoint.IsValid())
 	{
-		STORM_SYNC_CLIENT_LOG(Error, TEXT("FStormSyncPushMessageService::HandlePushResponseMessage - Invalid local endpoint."))
+		UE_LOG(LogStormSyncClient, Error, TEXT("FStormSyncPushMessageService::HandlePushResponseMessage - Invalid local endpoint."));
 		return;	
 	}
 	
@@ -209,7 +209,7 @@ void FStormSyncPushMessageService::HandlePushResponseMessage(const FStormSyncTra
 		// If text is non empty, indicates an error happened
 		if (!Payload->bSuccess)
 		{
-			STORM_SYNC_CLIENT_LOG(Error, TEXT("FStormSyncPushMessageService::HandlePushResponseMessage - %s"), *Payload->ErrorText.ToString())
+			UE_LOG(LogStormSyncClient, Error, TEXT("FStormSyncPushMessageService::HandlePushResponseMessage - %s"), *Payload->ErrorText.ToString());
 			
 			const TSharedPtr<FStormSyncTransportPushResponse> ErrorResponse = MakeShared<FStormSyncTransportPushResponse>(IncomingMessage);
 
@@ -224,7 +224,7 @@ void FStormSyncPushMessageService::HandlePushResponseMessage(const FStormSyncTra
 	if (!Connection.IsValid())
 	{
 		const FText ErrorText = LOCTEXT("Error_Active_Connection", "Returned active connection is invalid from endpoint StartSendingBuffer()");
-		STORM_SYNC_CLIENT_LOG(Error, TEXT("FStormSyncPushMessageService::HandlePushResponseMessage - %s"), *ErrorText.ToString())
+		UE_LOG(LogStormSyncClient, Error, TEXT("FStormSyncPushMessageService::HandlePushResponseMessage - %s"), *ErrorText.ToString());
 
 		// Something went wrong with the tcp connection, send it back as an error to remote requester
 		const TSharedPtr<FStormSyncTransportPushResponse> ErrorResponse = MakeShared<FStormSyncTransportPushResponse>(InMessage);
@@ -241,7 +241,7 @@ void FStormSyncPushMessageService::HandlePushResponseMessage(const FStormSyncTra
 
 void FStormSyncPushMessageService::HandleTransferComplete(const FStormSyncTransportPushResponse InResponseMessage)
 {
-	STORM_SYNC_CLIENT_LOG(Display, TEXT("FStormSyncPushMessageService::HandleTransferComplete - TCP transfer done for message with ID: %s"), *InResponseMessage.MessageId.ToString())
+	UE_LOG(LogStormSyncClient, Display, TEXT("FStormSyncPushMessageService::HandleTransferComplete - TCP transfer done for message with ID: %s"), *InResponseMessage.MessageId.ToString());
 	
 	const TSharedPtr<FStormSyncTransportPushResponse> Response =  MakeShared<FStormSyncTransportPushResponse>(InResponseMessage);
 	check(Response.IsValid());
@@ -254,7 +254,7 @@ void FStormSyncPushMessageService::HandleTransferComplete(const FStormSyncTransp
 void FStormSyncPushMessageService::InvokePendingResponseHandler(const TSharedPtr<FStormSyncTransportPushResponse, ESPMode::ThreadSafe>& InResponseMessage)
 {
 	check(InResponseMessage.IsValid());
-	STORM_SYNC_CLIENT_LOG(Display, TEXT("FStormSyncPushMessageService::InvokePendingResponseHandler - Trying to invoke callback (if any) for message with ID: %s"), *InResponseMessage->ToString())
+	UE_LOG(LogStormSyncClient, Display, TEXT("FStormSyncPushMessageService::InvokePendingResponseHandler - Trying to invoke callback (if any) for message with ID: %s"), *InResponseMessage->ToString());
 	
 	FScopeLock ScopeLock(&HandlersCriticalSection);
 	const FOnStormSyncPushComplete* ResponseHandler = PendingResponseHandlers.Find(InResponseMessage->MessageId);
diff --git a/Engine/Plugins/Experimental/StormSync/Source/StormSyncTransportClient/Private/Services/StormSyncStatusMessageService.cpp b/Engine/Plugins/Experimental/StormSync/Source/StormSyncTransportClient/Private/Services/StormSyncStatusMessageService.cpp
--- a/Engine/Plugins/Experimental/StormSync/Source/StormSyncTransportClient/Private/Services/StormSyncStatusMessageService.cpp
+++ b/Engine/Plugins/Experimental/StormSync/Source/StormSyncTransportClient/Private/Services/StormSyncStatusMessageService.cpp
@@ -40,7 +40,7 @@ void FStormSyncStatusMessageService::RequestStatus(const FMessageAddress& InRemo
 		const TArray<FStormSyncFileDependency> FileDependencies = Result.Get();
 		if (FileDependencies.IsEmpty())
 		{
-			STORM_SYNC_CLIENT_LOG(Error, TEXT("FStormSyncStatusMessageService::RequestStatus - Async FileDependencies is empty, something went wrong"))
+			UE_LOG(LogStormSyncClient, Error, TEXT("FStormSyncStatusMessageService::RequestStatus - Async FileDependencies is empty, something went wrong"));
 			return;
 		}
 
@@ -53,7 +53,7 @@ void FStormSyncStatusMessageService::RequestStatus(const FMessageAddress& InRemo
 		const TSharedPtr<FMessageEndpoint, ESPMode::ThreadSafe> MessageEndpoint = LocalThis->GetMessageEndpoint();
 		if (!MessageEndpoint.IsValid())
 		{
-			STORM_SYNC_CLIENT_LOG(Error, TEXT("FStormSyncStatusMessageService::RequestStatus - Unable to get message endpoint"))
+			UE_LOG(LogStormSyncClient, Error, TEXT("FStormSyncStatusMessageService::RequestStatus - Unable to get message endpoint"));
 			return;
 		}
 
@@ -61,15 +61,15 @@ void FStormSyncStatusMessageService::RequestStatus(const FMessageAddress& InRemo
 		TUniquePtr<FStormSyncTransportStatusRequest> RequestStatusMessage(FMessageEndpoint::MakeMessage<FStormSyncTransportStatusRequest>(PackageNames, FileDependencies));
 		if (!RequestStatusMessage.IsValid())
 		{
-			STORM_SYNC_CLIENT_LOG(Error, TEXT("FStormSyncStatusMessageService::RequestStatus - Push request message is invalid"))
+			UE_LOG(LogStormSyncClient, Error, TEXT("FStormSyncStatusMessageService::RequestStatus - Push request message is invalid"));
 			return;
 		}
 
 		// Store completion callback for later invocation when receiving back a status response
 		LocalThis->AddResponseHandler(RequestStatusMessage->MessageId, DoneDelegate);
 
-		STORM_SYNC_CLIENT_LOG(Display, TEXT("FStormSyncStatusMessageService::RequestStatus - FileDependencies: %d"), FileDependencies.Num())
-		STORM_SYNC_CLIENT_LOG(Display, TEXT("FStormSyncStatusMessageService::RequestStatus - Message: %s"), *RequestStatusMessage->ToString())
+		UE_LOG(LogStormSyncClient, Display, TEXT("FStormSyncStatusMessageService::RequestStatus - FileDependencies: %d"), FileDependencies.Num());
+		UE_LOG(LogStormSyncClient, Display, TEXT("FStormSyncStatusMessageService::RequestStatus - Message: %s"), *RequestStatusMessage->ToString());
 
 		MessageEndpoint->Send(RequestStatusMessage.Release(), RemoteAddress);
 	});
@@ -92,25 +92,25 @@ void FStormSyncStatusMessageService::HandleStatusPing(const FStormSyncTransportS
 	const TSharedPtr<FMessageEndpoint, ESPMode::ThreadSafe> MessageEndpoint = GetMessageEndpoint();
 	if (!MessageEndpoint.IsValid())
 	{
-		STORM_SYNC_CLIENT_LOG(Error, TEXT("FStormSyncStatusMessageService::HandleStatusPing - Unable to get client endpoint"))
+		UE_LOG(LogStormSyncClient, Error, TEXT("FStormSyncStatusMessageService::HandleStatusPing - Unable to get client endpoint"));
 		return;
 	}
 
-	STORM_SYNC_CLIENT_LOG(Display, TEXT("FStormSyncStatusMessageService::HandleStatusPing - Sender: %s"), *MessageContext->GetSender().ToString())
+	UE_LOG(LogStormSyncClient, Display, TEXT("FStormSyncStatusMessageService::HandleStatusPing - Sender: %s"), *MessageContext->GetSender().ToString());
 
 	// Send back pong message so that this recipient knows about this editor instance
-	STORM_SYNC_CLIENT_LOG(Display, TEXT("FStormSyncStatusMessageService::HandleStatusPing - Send pong Message back to %s..."), *MessageContext->GetSender().ToString())
+	UE_LOG(LogStormSyncClient, Display, TEXT("FStormSyncStatusMessageService::HandleStatusPing - Send pong Message back to %s..."), *MessageContext->GetSender().ToString());
 	FStormSyncTransportStatusPong* Message = FMessageEndpoint::MakeMessage<FStormSyncTransportStatusPong>();
 	MessageEndpoint->Send(Message, MessageContext->GetSender());
 }
 
 void FStormSyncStatusMessageService::HandleStatusRequest(const FStormSyncTransportStatusRequest& InMessage, const TSharedRef<IMessageContext, ESPMode::ThreadSafe>& MessageContext)
 {
-	STORM_SYNC_CLIENT_LOG(Display, TEXT("FStormSyncStatusMessageService::HandleStatusRequest - Message: %s, Sender: %s"), *InMessage.ToString(), *MessageContext->GetSender().ToString())
+	UE_LOG(LogStormSyncClient, Display, TEXT("FStormSyncStatusMessageService::HandleStatusRequest - Message: %s, Sender: %s"), *InMessage.ToString(), *MessageContext->GetSender().ToString());
 	const TSharedPtr<FMessageEndpoint, ESPMode::ThreadSafe> MessageEndpoint = GetMessageEndpoint();
 	if (!MessageEndpoint.IsValid())
 	{
-		STORM_SYNC_CLIENT_LOG(Error, TEXT("FStormSyncStatusMessageService::HandleStatusPing - Unable to get client endpoint"))
+		UE_LOG(LogStormSyncClient, Error, TEXT("FStormSyncStatusMessageService::HandleStatusPing - Unable to get client endpoint"));
 		return;
 	}
 
@@ -120,10 +120,10 @@ void FStormSyncStatusMessageService::HandleStatusRequest(const FStormSyncTranspo
 	Message->Modifiers = FStormSyncCoreUtils::GetSyncFileModifiers(InMessage.PackageNames, InMessage.Dependencies);
 	Message->bNeedsSynchronization = !Message->Modifiers.IsEmpty();
 
-	STORM_SYNC_CLIENT_LOG(Display, TEXT("GetSyncFileModifiers ... - Modifiers: %d"), Message->Modifiers.Num())
+	UE_LOG(LogStormSyncClient, Display, TEXT("GetSyncFileModifiers ... - Modifiers: %d"), Message->Modifiers.Num());
 	for (const FStormSyncFileModifierInfo& Modifier : Message->Modifiers)
 	{
-		STORM_SYNC_CLIENT_LOG(Display, TEXT("\t Modifier: %s"), *Modifier.ToString())
+		UE_LOG(LogStormSyncClient, Display, TEXT("\t Modifier: %s"), *Modifier.ToString());
 	}
 
 	// Send back result
@@ -132,13 +132,13 @@ void FStormSyncStatusMessageService::HandleStatusRequest(const FStormSyncTranspo
 
 void FStormSyncStatusMessageService::HandleStatusResponse(const FStormSyncTransportStatusResponse& InMessage, const TSharedRef<IMessageContext, ESPMode::ThreadSafe>& MessageContext)
 {
-	STORM_SYNC_CLIENT_LOG(Display, TEXT("FStormSyncStatusMessageService::HandleStatusResponse - Message: %s, Sender: %s (Id: %s)"), *InMessage.ToString(), *MessageContext->GetSender().ToString(), *InMessage.StatusRequestId.ToString())
+	UE_LOG(LogStormSyncClient, Display, TEXT("FStormSyncStatusMessageService::HandleStatusResponse - Message: %s, Sender: %s (Id: %s)"), *InMessage.ToString(), *MessageContext->GetSender().ToString(), *InMessage.StatusRequestId.ToString());
 
 	TArray<FStormSyncFileModifierInfo> NewModifiers;
 	TArray<FStormSyncFileModifierInfo> Modifiers = InMessage.Modifiers;
 	for (const FStormSyncFileModifierInfo& Modifier : Modifiers)
 	{
-		STORM_SYNC_CLIENT_LOG(Display, TEXT("\t Original Modifier: %s"), *Modifier.ToString())
+		UE_LOG(LogStormSyncClient, Display, TEXT("\t Original Modifier: %s"), *Modifier.ToString());
 
 		// Handle "missing" case, meaning remote is referencing the file but we are not
 		// It doesn't mean we don't have the file though, check for file equality here
@@ -149,7 +149,7 @@ void FStormSyncStatusMessageService::HandleStatusResponse(const FStormSyncTransp
 			// Add it to modifiers if we could not find it locally
 			if (!FileDependency.IsValid())
 			{
-				STORM_SYNC_CLIENT_LOG(Display, TEXT("\t\t Adding as missing locally, Modifier: %s"), *Modifier.ToString())
+				UE_LOG(LogStormSyncClient, Display, TEXT("\t\t Adding as missing locally, Modifier: %s"), *Modifier.ToString());
 				NewModifiers.Add(Modifier);
 				continue;
 			}
@@ -162,7 +162,7 @@ void FStormSyncStatusMessageService::HandleStatusResponse(const FStormSyncTransp
 
 			if (!bSameSize || !bSameHash)
 			{
-				STORM_SYNC_CLIENT_LOG(Display, TEXT("\t\t Mismatched state locally, add it as overwrite operation. Modifier: %s"), *Modifier.ToString())
+				UE_LOG(LogStormSyncClient, Display, TEXT("\t\t Mismatched state locally, add it as overwrite operation. Modifier: %s"), *Modifier.ToString());
 
 				// Mismatched state, change modifier operation to overwrite
 				NewModifiers.Add(FStormSyncFileModifierInfo(EStormSyncModifierOperation::Overwrite, RemoteDependency));
@@ -174,10 +174,10 @@ void FStormSyncStatusMessageService::HandleStatusResponse(const FStormSyncTransp
 		}
 	}
 
-	STORM_SYNC_CLIENT_LOG(Display, TEXT("FStormSyncStatusMessageService::HandleStatusResponse - NewModifiers: %d"), NewModifiers.Num())
+	UE_LOG(LogStormSyncClient, Display, TEXT("FStormSyncStatusMessageService::HandleStatusResponse - NewModifiers: %d"), NewModifiers.Num());
 	for (const FStormSyncFileModifierInfo& Modifier : NewModifiers)
 	{
-		STORM_SYNC_CLIENT_LOG(Display, TEXT("\t New Modifier: %s"), *Modifier.ToString())
+		UE_LOG(LogStormSyncClient, Display, TEXT("\t New Modifier: %s"), *Modifier.ToString());
 	}
 
 	{
diff --git a/Engine/Plugins/Experimental/StormSync/Source/StormSyncTransportClient/Private/Socket/StormSyncTransportClientSocket.cpp b/Engine/Plugins/Experimental/StormSync/Source/StormSyncTransportClient/Private/Socket/StormSyncTransportClientSocket.cpp
--- a/Engine/Plugins/Experimental/StormSync/Source/StormSyncTransportClient/Private/Socket/StormSyncTransportClientSocket.cpp
+++ b/Engine/Plugins/Experimental/StormSync/Source/StormSyncTransportClient/Private/Socket/StormSyncTransportClientSocket.cpp
@@ -24,7 +24,7 @@ FStormSyncTransportClientSocket::FStormSyncTransportClientSocket(const FIPv4Endp
 
 FStormSyncTransportClientSocket::~FStormSyncTransportClientSocket()
 {
-	STORM_SYNC_CLIENT_LOG(Display, TEXT("FStormSyncTransportClientSocket::~FStormSyncTransportClientSocket Destructor Closing socket and runnable to '%s'"), *RemoteEndpoint.ToString());
+	UE_LOG(LogStormSyncClient, Display, TEXT("FStormSyncTransportClientSocket::~FStormSyncTransportClientSocket Destructor Closing socket and runnable to '%s'"), *RemoteEndpoint.ToString());
 	
 	if (Thread != nullptr)
 	{
@@ -61,7 +61,7 @@ void FStormSyncTransportClientSocket::StartTransport()
 
 void FStormSyncTransportClientSocket::StopTransport()
 {
-	STORM_SYNC_CLIENT_LOG(Display, TEXT("FStormSyncTransportClientSocket::StopTransport Closing socket and runnable to '%s'"), *RemoteEndpoint.ToString());
+	UE_LOG(LogStormSyncClient, Display, TEXT("FStormSyncTransportClientSocket::StopTransport Closing socket and runnable to '%s'"), *RemoteEndpoint.ToString());
 	
 	// let the thread shutdown on its own
 	if (Thread != nullptr)
@@ -96,7 +96,7 @@ void FStormSyncTransportClientSocket::Stop()
 {
 	if (Socket)
 	{
-		STORM_SYNC_CLIENT_LOG(Display, TEXT("FStormSyncTransportClientSocket::Stop Closing socket to '%s'"), *RemoteEndpoint.ToString());
+		UE_LOG(LogStormSyncClient, Display, TEXT("FStormSyncTransportClientSocket::Stop Closing socket to '%s'"), *RemoteEndpoint.ToString());
 		Socket->Close();
 		ConnectionClosedDelegate.ExecuteIfBound(RemoteEndpoint);
 	}
@@ -104,7 +104,7 @@ void FStormSyncTransportClientSocket::Stop()
 
 uint32 FStormSyncTransportClientSocket::Run()
 {
-	STORM_SYNC_CLIENT_LOG(Log, TEXT("Started Connection to '%s'"), *RemoteEndpoint.ToString());
+	UE_LOG(LogStormSyncClient, Log, TEXT("Started Connection to '%s'"), *RemoteEndpoint.ToString());
 
 	double LastTime = FPlatformTime::Seconds();
 	constexpr float IdealFrameTime = 1.0f / 30.0f;
@@ -119,13 +119,13 @@ uint32 FStormSyncTransportClientSocket::Run()
 
 		if (DeltaTime > 1.2f)
 		{
-			STORM_SYNC_CLIENT_LOG(Warning, TEXT("Hitch detected; %.3f seconds since prior tick"), DeltaTime);
+			UE_LOG(LogStormSyncClient, Warning, TEXT("Hitch detected; %.3f seconds since prior tick"), DeltaTime);
 		}
 
 		FString StopReason;
 		bListenerIsRunning = Tick(StopReason);
 		
-		STORM_SYNC_CLIENT_LOG(
+		UE_LOG(LogStormSyncClient, 
 			Verbose,
 			TEXT("Running ... Connected: %s (%s), State: %s, bListenerIsRunning: %s %s"),
 			IsConnected() ? TEXT("true") : TEXT("false"),
@@ -142,7 +142,7 @@ uint32 FStormSyncTransportClientSocket::Run()
 
 		if (!bListenerIsRunning)
 		{
-			STORM_SYNC_CLIENT_LOG(Display, TEXT("Stopped. Reason: %s"), *StopReason);
+			UE_LOG(LogStormSyncClient, Display, TEXT("Stopped. Reason: %s"), *StopReason);
 		}
 
 		// Throttle main thread main fps by sleeping if we still have time
@@ -168,7 +168,7 @@ bool FStormSyncTransportClientSocket::Tick(FString& StopReason)
 		{
 			const FString ErrorCode = GetSocketReadableErrorCode();
 			StopReason = FString::Printf(TEXT("Dummy read failed with code %s. Socket has closed [%s]"), *ErrorCode, *RemoteEndpoint.ToString());
-			STORM_SYNC_CLIENT_LOG(Verbose, TEXT("%s"), *StopReason);
+			UE_LOG(LogStormSyncClient, Verbose, TEXT("%s"), *StopReason);
 
 			{
 				FScopeLock SendLock(&SendCriticalSection);
@@ -204,7 +204,7 @@ bool FStormSyncTransportClientSocket::Tick(FString& StopReason)
 		{
 			const FString ErrorCode = GetSocketReadableErrorCode();
 			StopReason = FString::Printf(TEXT("Read failed with code %s [%s]"), *ErrorCode, *RemoteEndpoint.ToString());
-			STORM_SYNC_CLIENT_LOG(Verbose, TEXT("%s"), *StopReason);
+			UE_LOG(LogStormSyncClient, Verbose, TEXT("%s"), *StopReason);
 			return false;
 		}
 
@@ -270,7 +270,7 @@ bool FStormSyncTransportClientSocket::TryReconnect()
 bool FStormSyncTransportClientSocket::Connect()
 {
 	bool bSuccess = false;
-	STORM_SYNC_CLIENT_LOG(Verbose, TEXT("Connection to '%s', trying..."), *RemoteEndpoint.ToString());
+	UE_LOG(LogStormSyncClient, Verbose, TEXT("Connection to '%s', trying..."), *RemoteEndpoint.ToString());
     
 	Socket = FTcpSocketBuilder(SocketDescription)
 		.WithSendBufferSize(SocketBufferSize)
@@ -292,7 +292,7 @@ bool FStormSyncTransportClientSocket::Connect()
 			}
 			ConnectionStateChangedDelegate.ExecuteIfBound();
 
-			STORM_SYNC_CLIENT_LOG(
+			UE_LOG(LogStormSyncClient, 
 				Display,
 				TEXT("Connection to %s successful. Connection State: %s"),
 				*RemoteEndpoint.ToString(),
@@ -304,15 +304,15 @@ bool FStormSyncTransportClientSocket::Connect()
 			bSuccess = false;
 			
 			const FString SocketErrorCode = GetSocketReadableErrorCode();
-			STORM_SYNC_CLIENT_LOG(
+			UE_LOG(LogStormSyncClient, 
 				Warning,
 				TEXT("Socket connected to %s with invalid state: %s (Last Error Code: %s)"),
 				*RemoteEndpoint.ToString(),
 				*GetSocketReadableConnectionState(SocketConnectionState),
 				*SocketErrorCode
 			);
-			STORM_SYNC_CLIENT_LOG(Warning, TEXT("\t- Ensure the connection is on the correct subnet."))
-			STORM_SYNC_CLIENT_LOG(Warning, TEXT("\t- Ensure the firewall settings are not blocking the connection."))
+			UE_LOG(LogStormSyncClient, Warning, TEXT("\t- Ensure the connection is on the correct subnet."));
+			UE_LOG(LogStormSyncClient, Warning, TEXT("\t- Ensure the firewall settings are not blocking the connection."));
 		}
 	}
 
@@ -332,7 +332,7 @@ bool FStormSyncTransportClientSocket::Connect()
 		// Block waiting for some data
 		if (!Socket->Wait(ESocketWaitConditions::WaitForRead, FTimespan::FromSeconds(SocketWaitForReadTimeSeconds)))
 		{
-			STORM_SYNC_CLIENT_LOG(
+			UE_LOG(LogStormSyncClient, 
 				Warning,
 				TEXT("Connect wait for state packet failed. We were able to connect to %s but server didn't respond (Last Error Code: %s)"),
 				*RemoteEndpoint.ToString(),
@@ -349,7 +349,7 @@ bool FStormSyncTransportClientSocket::Connect()
 			int32 BytesRead = 0;
 			if (!Socket->Recv(Buffer.GetData(), PendingDataSize, BytesRead, ESocketReceiveFlags::None))
 			{
-				STORM_SYNC_CLIENT_LOG(Error, TEXT("Error while receiving data via endpoint %s"), *RemoteEndpoint.ToString());
+				UE_LOG(LogStormSyncClient, Error, TEXT("Error while receiving data via endpoint %s"), *RemoteEndpoint.ToString());
 				return false;
 			}
 
@@ -374,7 +374,7 @@ void FStormSyncTransportClientSocket::CloseSocket(const bool bShouldTriggerEvent
 {
 	if (Socket)
 	{
-		STORM_SYNC_CLIENT_LOG(Display, TEXT("FStormSyncTransportClientSocket::CloseSocket Closing socket to '%s'"), *RemoteEndpoint.ToString());
+		UE_LOG(LogStormSyncClient, Display, TEXT("FStormSyncTransportClientSocket::CloseSocket Closing socket to '%s'"), *RemoteEndpoint.ToString());
 		Socket->Close();
 
 		if (bShouldTriggerEvent)
@@ -398,14 +398,14 @@ void FStormSyncTransportClientSocket::SendBuffer(const TArray<uint8>& InBuffer)
 
 	const EConnectionState State = GetConnectionState();
 
-	STORM_SYNC_CLIENT_LOG(
+	UE_LOG(LogStormSyncClient, 
 		Display,
 		TEXT("FStormSyncTransportClientSocket::SendBuffer - Sending buffer of size %d to %s:%d endpoint (State: %s)"),
 		BufferSize,
 		*RemoteEndpoint.Address.ToString(),
 		RemoteEndpoint.Port,
 		*GetReadableConnectionState(State)
-	)
+	);
 
 	// The first 4 bytes are going to be our expected buffer size on the server end
 	bSending = Send(InBuffer.GetData(), BufferSize);
@@ -432,13 +432,13 @@ bool FStormSyncTransportClientSocket::Send(const uint8* Data, uint32 Size, const
 {
 	if (!Socket)
 	{
-		STORM_SYNC_CLIENT_LOG(Error, TEXT("FStormSyncTransportClientSocket::Send - Socket is not active. Make sure to call Connect() prior to sending buffers."))
+		UE_LOG(LogStormSyncClient, Error, TEXT("FStormSyncTransportClientSocket::Send - Socket is not active. Make sure to call Connect() prior to sending buffers."));
 		return false;
 	}
 	
 	if (ConnectionState != State_Connected)
 	{
-		STORM_SYNC_CLIENT_LOG(
+		UE_LOG(LogStormSyncClient, 
 			Warning,
 			TEXT("FStormSyncTransportClientSocket::Send - Connection State seems to be invalid: %s (Socket State: %s, Last Socket Error: %s)"),
 			*GetReadableConnectionState(ConnectionState),
@@ -446,7 +446,7 @@ bool FStormSyncTransportClientSocket::Send(const uint8* Data, uint32 Size, const
 			*GetSocketReadableErrorCode()
 		);
 		
-		STORM_SYNC_CLIENT_LOG(Warning, TEXT("Attempted to send a socket message but no socket connected, ignoring..."));
+		UE_LOG(LogStormSyncClient, Warning, TEXT("Attempted to send a socket message but no socket connected, ignoring..."));
 		return false;
 	}
 	
@@ -468,11 +468,11 @@ bool FStormSyncTransportClientSocket::Send(const uint8* Data, uint32 Size, const
 
 	Buffer.Append(Data, Size);
 
-	STORM_SYNC_CLIENT_LOG(Display, TEXT("Client: Sending %d bytes (Size: %d) ..."), Buffer.Num(), Size);
+	UE_LOG(LogStormSyncClient, Display, TEXT("Client: Sending %d bytes (Size: %d) ..."), Buffer.Num(), Size);
 	if (!BlockingSend(Buffer.GetData(), Buffer.Num()))
 	{
 		const FString ErrorCode = GetSocketReadableErrorCode();
-		STORM_SYNC_CLIENT_LOG(Warning, TEXT("FStormSyncTransportClientSocket::BlockingSend failed with code %s"), *ErrorCode);
+		UE_LOG(LogStormSyncClient, Warning, TEXT("FStormSyncTransportClientSocket::BlockingSend failed with code %s"), *ErrorCode);
 		return false;
 	}
 
@@ -525,7 +525,7 @@ void FStormSyncTransportClientSocket::ParseIncomingBytes(const int32 InNumBytesR
 
 bool FStormSyncTransportClientSocket::ParseIncomingMessage(const FString& InMessage, const FIPv4Endpoint& InEndpoint)
 {
-	STORM_SYNC_CLIENT_LOG(VeryVerbose, TEXT("FStormSyncTransportClientSocket::ParseIncomingMessage from %s - Message: %s"), *InEndpoint.ToString(), *InMessage);
+	UE_LOG(LogStormSyncClient, VeryVerbose, TEXT("FStormSyncTransportClientSocket::ParseIncomingMessage from %s - Message: %s"), *InEndpoint.ToString(), *InMessage);
 
 	const TSharedRef<TJsonReader<TCHAR>> Reader = FJsonStringReader::Create(InMessage);
 	TSharedPtr<FJsonObject> JsonData;
@@ -550,7 +550,7 @@ bool FStormSyncTransportClientSocket::ParseIncomingMessage(const FString& InMess
 		double ReceivedBytes;
 		if (!JsonData->TryGetNumberField(FieldName, ReceivedBytes))
 		{
-			STORM_SYNC_CLIENT_LOG(
+			UE_LOG(LogStormSyncClient, 
 				Error,
 				TEXT("FStormSyncTransportClientSocket::ParseIncomingMessage \"size\" command missing required field \"%s\" or not able to convert to number"),
 				*FieldName
diff --git a/Engine/Plugins/Experimental/StormSync/Source/StormSyncTransportClient/Private/StormSyncTransportClientEndpoint.cpp b/Engine/Plugins/Experimental/StormSync/Source/StormSyncTransportClient/Private/StormSyncTransportClientEndpoint.cpp
--- a/Engine/Plugins/Experimental/StormSync/Source/StormSyncTransportClient/Private/StormSyncTransportClientEndpoint.cpp
+++ b/Engine/Plugins/Experimental/StormSync/Source/StormSyncTransportClient/Private/StormSyncTransportClientEndpoint.cpp
@@ -34,20 +34,20 @@ void FStormSyncTransportClientEndpoint::InitializeMessaging(const FString& InEnd
 	// Message endpoint already created
 	if (MessageEndpoint.IsValid())
 	{
-		STORM_SYNC_CLIENT_LOG(Warning, TEXT("FStormSyncTransportClientEndpoint::InitializeMessaging - Message endpoint already initialized"))
+		UE_LOG(LogStormSyncClient, Warning, TEXT("FStormSyncTransportClientEndpoint::InitializeMessaging - Message endpoint already initialized"));
 		return;
 	}
 
 	const IMessageBusPtr MessageBus = MessageBusPtr.Pin();
 	if (!MessageBus.IsValid())
 	{
-		STORM_SYNC_CLIENT_LOG(Error, TEXT("FStormSyncTransportClientEndpoint::InitializeMessaging - Default message bus is invalid"))
+		UE_LOG(LogStormSyncClient, Error, TEXT("FStormSyncTransportClientEndpoint::InitializeMessaging - Default message bus is invalid"));
 		return;
 	}
 
 	// Create endpoint builder
 	const FString MessageEndpointName = FString::Printf(TEXT("StormSync%sEndpoint"), *InEndpointFriendlyName);
-	STORM_SYNC_CLIENT_LOG(Verbose, TEXT("FStormSyncTransportClientEndpoint::InitializeMessaging - Setting up %s message endpoint"), *MessageEndpointName)
+	UE_LOG(LogStormSyncClient, Verbose, TEXT("FStormSyncTransportClientEndpoint::InitializeMessaging - Setting up %s message endpoint"), *MessageEndpointName);
 
 	FMessageEndpointBuilder EndpointBuilder = FMessageEndpoint::Builder(*MessageEndpointName, MessageBus.ToSharedRef())
 		.ReceivingOnThread(ENamedThreads::Type::GameThread);
@@ -69,7 +69,7 @@ void FStormSyncTransportClientEndpoint::InitializeMessaging(const FString& InEnd
 	check(MessageEndpoint.IsValid());
 
 	// Message subscribes
-	STORM_SYNC_CLIENT_LOG(Display, TEXT("FStormSyncTransportClientEndpoint::InitializeMessaging - Subscribe to messages"))
+	UE_LOG(LogStormSyncClient, Display, TEXT("FStormSyncTransportClientEndpoint::InitializeMessaging - Subscribe to messages"));
 	MessageEndpoint->Subscribe<FStormSyncTransportStatusPing>();
 }
 
@@ -105,7 +105,7 @@ void FStormSyncTransportClientEndpoint::RequestStatus(const FMessageAddress& InR
 {
 	if (!StatusService.IsValid())
 	{
-		STORM_SYNC_CLIENT_LOG(Error, TEXT("FStormSyncTransportClientEndpoint::RequestStatus - Status message service is invalid"))
+		UE_LOG(LogStormSyncClient, Error, TEXT("FStormSyncTransportClientEndpoint::RequestStatus - Status message service is invalid"));
 		return;
 	}
 
@@ -116,7 +116,7 @@ void FStormSyncTransportClientEndpoint::RequestPushPackages(const FMessageAddres
 {
 	if (!PushService.IsValid())
 	{
-		STORM_SYNC_CLIENT_LOG(Error, TEXT("FStormSyncTransportClientEndpoint::RequestPushPackages - Push message service is invalid"))
+		UE_LOG(LogStormSyncClient, Error, TEXT("FStormSyncTransportClientEndpoint::RequestPushPackages - Push message service is invalid"));
 		return;
 	}
 
@@ -127,7 +127,7 @@ void FStormSyncTransportClientEndpoint::RequestPullPackages(const FMessageAddres
 {
 	if (!PullService.IsValid())
 	{
-		STORM_SYNC_CLIENT_LOG(Error, TEXT("FStormSyncTransportClientEndpoint::RequestPullPackages - Pull message service is invalid"))
+		UE_LOG(LogStormSyncClient, Error, TEXT("FStormSyncTransportClientEndpoint::RequestPullPackages - Pull message service is invalid"));
 		return;
 	}
 
@@ -138,7 +138,7 @@ void FStormSyncTransportClientEndpoint::AbortStatusRequest(const FGuid& InStatus
 {
 	if (!StatusService.IsValid())
 	{
-		STORM_SYNC_CLIENT_LOG(Error, TEXT("FStormSyncTransportClientEndpoint::AbortStatusRequest - Status message service is invalid"))
+		UE_LOG(LogStormSyncClient, Error, TEXT("FStormSyncTransportClientEndpoint::AbortStatusRequest - Status message service is invalid"));
 		return;
 	}
 
@@ -149,7 +149,7 @@ void FStormSyncTransportClientEndpoint::AbortPushRequest(const FGuid& InPushRequ
 {
 	if (!PushService.IsValid())
 	{
-		STORM_SYNC_CLIENT_LOG(Error, TEXT("FStormSyncTransportClientEndpoint::AbortPushRequest - Status message service is invalid"))
+		UE_LOG(LogStormSyncClient, Error, TEXT("FStormSyncTransportClientEndpoint::AbortPushRequest - Status message service is invalid"));
 		return;
 	}
 
@@ -160,7 +160,7 @@ void FStormSyncTransportClientEndpoint::AbortPullRequest(const FGuid& InPullRequ
 {
 	if (!PullService.IsValid())
 	{
-		STORM_SYNC_CLIENT_LOG(Error, TEXT("FStormSyncTransportClientEndpoint::AbortPullRequest - Status message service is invalid"))
+		UE_LOG(LogStormSyncClient, Error, TEXT("FStormSyncTransportClientEndpoint::AbortPullRequest - Status message service is invalid"));
 		return;
 	}
 
@@ -169,7 +169,7 @@ void FStormSyncTransportClientEndpoint::AbortPullRequest(const FGuid& InPullRequ
 
 void FStormSyncTransportClientEndpoint::StartSendingBuffer(const FStormSyncTransportSyncResponse& InMessage, TSharedPtr<FStormSyncTransportClientSocket>& OutActiveConnection, const FOnStormSyncSendBufferCallback& InDoneDelegate)
 {
-	STORM_SYNC_CLIENT_LOG(Display, TEXT("FStormSyncTransportClientEndpoint::StartSendingBuffer - Message: %s"), *InMessage.ToString())
+	UE_LOG(LogStormSyncClient, Display, TEXT("FStormSyncTransportClientEndpoint::StartSendingBuffer - Message: %s"), *InMessage.ToString());
 	
 	if (InMessage.Modifiers.IsEmpty())
 	{
@@ -225,7 +225,7 @@ void FStormSyncTransportClientEndpoint::StartSendingBuffer(const FStormSyncTrans
 		TArray<uint8> PakBuffer;
 		if (!FStormSyncCoreUtils::CreatePakBuffer(PackageNames, PakBuffer, ErrorText))
 		{
-			STORM_SYNC_CLIENT_LOG(Display, TEXT("FStormSyncTransportClientEndpoint::StartSendingBuffer - CreatePakBuffer Error: %s"), *ErrorText.ToString())
+			UE_LOG(LogStormSyncClient, Display, TEXT("FStormSyncTransportClientEndpoint::StartSendingBuffer - CreatePakBuffer Error: %s"), *ErrorText.ToString());
 			InDoneDelegate.ExecuteIfBound(MakeShared<FStormSyncSendingBufferPayload>(ErrorText));
 			return;
 		}
@@ -237,7 +237,7 @@ void FStormSyncTransportClientEndpoint::StartSendingBuffer(const FStormSyncTrans
 			return;
 		}
 		
-		STORM_SYNC_CLIENT_LOG(Display, TEXT("FStormSyncTransportClientEndpoint::StartSendingBuffer - Try sending on %s"), *ServerAddress)
+		UE_LOG(LogStormSyncClient, Display, TEXT("FStormSyncTransportClientEndpoint::StartSendingBuffer - Try sending on %s"), *ServerAddress);
 		StrongThis->SendTcpBuffer(ServerAddress, PakBuffer);
 		
 		InDoneDelegate.ExecuteIfBound(MakeShared<FStormSyncSendingBufferPayload>());
@@ -249,7 +249,7 @@ TSharedPtr<FStormSyncTransportClientSocket> FStormSyncTransportClientEndpoint::G
 	FIPv4Endpoint Endpoint(FIPv4Address::Any, 0);
 	if (!FIPv4Endpoint::Parse(InAddress, Endpoint))
 	{
-		STORM_SYNC_CLIENT_LOG(Error, TEXT("FStormSyncTransportClientSocket::GetOrCreateClientSocket - Failed to parse endpoint '%s'."), *InAddress);
+		UE_LOG(LogStormSyncClient, Error, TEXT("FStormSyncTransportClientSocket::GetOrCreateClientSocket - Failed to parse endpoint '%s'."), *InAddress);
 		return nullptr;
 	}
 
@@ -273,7 +273,7 @@ TSharedPtr<FStormSyncTransportClientSocket> FStormSyncTransportClientEndpoint::G
 
 	if (!Connection->Connect())
 	{
-		STORM_SYNC_CLIENT_LOG(Error, TEXT("FStormSyncTransportClientSocket::GetOrCreateClientSocket - Connect failed on socket for %s"), *Endpoint.ToString());
+		UE_LOG(LogStormSyncClient, Error, TEXT("FStormSyncTransportClientSocket::GetOrCreateClientSocket - Connect failed on socket for %s"), *Endpoint.ToString());
 		return nullptr;
 	}
 
@@ -301,10 +301,10 @@ TSharedPtr<FStormSyncTransportClientSocket> FStormSyncTransportClientEndpoint::G
 	// Try establish connection on each, and stop on the first successful one
 	for (FString ServerAddress : ServerAddresses)
 	{
-		STORM_SYNC_CLIENT_LOG(Display, TEXT("FStormSyncTransportClientEndpoint::GetActiveConnection - Try connection on %s"), *ServerAddress)
+		UE_LOG(LogStormSyncClient, Display, TEXT("FStormSyncTransportClientEndpoint::GetActiveConnection - Try connection on %s"), *ServerAddress);
 		if (TSharedPtr<FStormSyncTransportClientSocket> Connection = GetOrCreateClientSocket(ServerAddress))
 		{
-			STORM_SYNC_CLIENT_LOG(Display, TEXT("FStormSyncTransportClientEndpoint::GetActiveConnection - Succesfful connection on %s"), *ServerAddress)
+			UE_LOG(LogStormSyncClient, Display, TEXT("FStormSyncTransportClientEndpoint::GetActiveConnection - Succesfful connection on %s"), *ServerAddress);
 			return Connection;
 		}
 	}
@@ -314,12 +314,12 @@ TSharedPtr<FStormSyncTransportClientSocket> FStormSyncTransportClientEndpoint::G
 
 void FStormSyncTransportClientEndpoint::HandlePongMessage(const FStormSyncTransportPongMessage& InMessage, const TSharedRef<IMessageContext, ESPMode::ThreadSafe>& InMessageContext)
 {
-	STORM_SYNC_CLIENT_LOG(Display, TEXT("FStormSyncTransportClientEndpoint::HandlePongMessage - Received pong Message %s from %s"), *InMessage.ToString(), *InMessageContext->GetSender().ToString())
+	UE_LOG(LogStormSyncClient, Display, TEXT("FStormSyncTransportClientEndpoint::HandlePongMessage - Received pong Message %s from %s"), *InMessage.ToString(), *InMessageContext->GetSender().ToString());
 }
 
 TArray<FName> FStormSyncTransportClientEndpoint::GetPackageNamesFromModifierInfos(const TArray<FStormSyncFileModifierInfo>& InModifierInfos)
 {
-	STORM_SYNC_CLIENT_LOG(Display, TEXT("FStormSyncTransportClientEndpoint::GetPackageNamesFromModifierInfos - ModifierInfos: %d"), InModifierInfos.Num())
+	UE_LOG(LogStormSyncClient, Display, TEXT("FStormSyncTransportClientEndpoint::GetPackageNamesFromModifierInfos - ModifierInfos: %d"), InModifierInfos.Num());
 
 	TArray<FName> PackageNames;
 	Algo::Transform(InModifierInfos, PackageNames, [](const FStormSyncFileModifierInfo& ModifierInfo)
@@ -335,20 +335,20 @@ void FStormSyncTransportClientEndpoint::SendTcpBuffer(const FString& InAddress,
 	const TSharedPtr<FStormSyncTransportClientSocket> Connection = GetOrCreateClientSocket(InAddress);
 	if (!Connection)
 	{
-		STORM_SYNC_CLIENT_LOG(Error, TEXT("FStormSyncTransportClientEndpoint::SendTcpBuffer - Failed to GetOrCreateClientSocket client socket to '%s'. TCP Client sending buffer will be disabled!"), *InAddress);
+		UE_LOG(LogStormSyncClient, Error, TEXT("FStormSyncTransportClientEndpoint::SendTcpBuffer - Failed to GetOrCreateClientSocket client socket to '%s'. TCP Client sending buffer will be disabled!"), *InAddress);
 		return;
 	}
 
 	const FIPv4Endpoint Endpoint = Connection->GetRemoteEndpoint();
 	const FStormSyncTransportClientSocket::EConnectionState State = Connection->GetConnectionState();
 
-	STORM_SYNC_CLIENT_LOG(
+	UE_LOG(LogStormSyncClient, 
 		Display,
 		TEXT("FStormSyncTransportClientEndpoint::SendTcpBuffer for %s (State: %s, Sending: %s)"),
 		*Endpoint.ToString(),
 		*Connection->GetReadableConnectionState(State),
 		Connection->IsSending() ? TEXT("true") : TEXT("false")
-	)
+	);
 
 	// Only allowed to send if we are connected and not currently sending on this connection
 	if (!Connection->IsSending())
@@ -366,17 +366,17 @@ void FStormSyncTransportClientEndpoint::SendTcpBuffer(const FString& InAddress,
 	}
 	else
 	{
-		STORM_SYNC_CLIENT_LOG(
+		UE_LOG(LogStormSyncClient, 
 			Error,
 			TEXT("FStormSyncTransportClientEndpoint::SendTcpBuffer - We are currently sending data to %s. Please wait for it to complete before trying to send new data again."),
 			*Endpoint.ToString()
-		)
+		);
 	}
 }
 
 void FStormSyncTransportClientEndpoint::OnConnectionClosedForSocket(const FIPv4Endpoint& InEndpoint)
 {
-	STORM_SYNC_CLIENT_LOG(Display, TEXT("FStormSyncTransportClientEndpoint::OnConnectionClosedForSocket - TCP socket closed for %s"), *InEndpoint.ToString())
+	UE_LOG(LogStormSyncClient, Display, TEXT("FStormSyncTransportClientEndpoint::OnConnectionClosedForSocket - TCP socket closed for %s"), *InEndpoint.ToString());
 }
 
 void FStormSyncTransportClientEndpoint::OnConnectionStateChanged(const TSharedPtr<FStormSyncTransportClientSocket> InConnection)
@@ -384,7 +384,7 @@ void FStormSyncTransportClientEndpoint::OnConnectionStateChanged(const TSharedPt
 	const FIPv4Endpoint Endpoint = InConnection->GetRemoteEndpoint();
 	const FStormSyncTransportClientSocket::EConnectionState State = InConnection->GetConnectionState();
 
-	STORM_SYNC_CLIENT_LOG(Display, TEXT("FStormSyncTransportClientEndpoint::OnConnectionStateChanged for %s (State: %s)"), *Endpoint.ToString(), *InConnection->GetReadableConnectionState(State))
+	UE_LOG(LogStormSyncClient, Display, TEXT("FStormSyncTransportClientEndpoint::OnConnectionStateChanged for %s (State: %s)"), *Endpoint.ToString(), *InConnection->GetReadableConnectionState(State));
 
 	const bool bHasConnectionRetryDelay = GetDefault<UStormSyncTransportSettings>()->HasConnectionRetryDelay();
 
@@ -392,7 +392,7 @@ void FStormSyncTransportClientEndpoint::OnConnectionStateChanged(const TSharedPt
 	if (!bHasConnectionRetryDelay && State == FStormSyncTransportClientSocket::State_Closed)
 	{
 		FScopeLock ScopeLock(&ConnectionsCriticalSection);
-		STORM_SYNC_CLIENT_LOG(Display, TEXT("FStormSyncTransportClientEndpoint::OnConnectionStateChanged - Closed, dispose of %s connection"), *InConnection->GetRemoteEndpoint().ToString());
+		UE_LOG(LogStormSyncClient, Display, TEXT("FStormSyncTransportClientEndpoint::OnConnectionStateChanged - Closed, dispose of %s connection"), *InConnection->GetRemoteEndpoint().ToString());
 		Connections.Remove(InConnection->GetRemoteEndpoint());
 	}
 }
@@ -404,7 +404,7 @@ void FStormSyncTransportClientEndpoint::OnConnectionReceivedBytes(const int32 In
 
 void FStormSyncTransportClientEndpoint::OnTransferComplete(const FString InConnectionAddress)
 {
-	STORM_SYNC_CLIENT_LOG(Display, TEXT("FStormSyncTransportClientEndpoint::OnTransferComplete - TCP socket done transfering for %s"), *InConnectionAddress)
+	UE_LOG(LogStormSyncClient, Display, TEXT("FStormSyncTransportClientEndpoint::OnTransferComplete - TCP socket done transfering for %s"), *InConnectionAddress);
 	FStormSyncCoreDelegates::OnTransferComplete.Broadcast(InConnectionAddress);
 }
 
diff --git a/Engine/Plugins/Experimental/StormSync/Source/StormSyncTransportClient/Private/StormSyncTransportClientLog.h b/Engine/Plugins/Experimental/StormSync/Source/StormSyncTransportClient/Private/StormSyncTransportClientLog.h
--- a/Engine/Plugins/Experimental/StormSync/Source/StormSyncTransportClient/Private/StormSyncTransportClientLog.h
+++ b/Engine/Plugins/Experimental/StormSync/Source/StormSyncTransportClient/Private/StormSyncTransportClientLog.h
@@ -5,8 +5,3 @@
 #include "Logging/LogMacros.h"
 
 DECLARE_LOG_CATEGORY_EXTERN(LogStormSyncClient, Display, All);
-
-#define STORM_SYNC_CLIENT_LOG(Verbosity, Format, ...) \
-{ \
-    UE_LOG(LogStormSyncClient, Verbosity, Format, ##__VA_ARGS__); \
-}
diff --git a/Engine/Plugins/Experimental/StormSync/Source/StormSyncTransportClient/Private/StormSyncTransportClientModule.cpp b/Engine/Plugins/Experimental/StormSync/Source/StormSyncTransportClient/Private/StormSyncTransportClientModule.cpp
--- a/Engine/Plugins/Experimental/StormSync/Source/StormSyncTransportClient/Private/StormSyncTransportClientModule.cpp
+++ b/Engine/Plugins/Experimental/StormSync/Source/StormSyncTransportClient/Private/StormSyncTransportClientModule.cpp
@@ -55,13 +55,13 @@ FMessageEndpointSharedPtr FStormSyncTransportClientModule::GetClientMessageEndpo
 {
 	if (!ClientEndpoint.IsValid())
 	{
-		STORM_SYNC_CLIENT_LOG(Warning, TEXT("FStormSyncTransportClientModule::GetClientMessageEndpoint - Client endpoint not valid"))
+		UE_LOG(LogStormSyncClient, Warning, TEXT("FStormSyncTransportClientModule::GetClientMessageEndpoint - Client endpoint not valid"));
 		return nullptr;
 	}
 
 	if (!ClientEndpoint->IsRunning())
 	{
-		STORM_SYNC_CLIENT_LOG(Warning, TEXT("FStormSyncTransportClientModule::GetClientMessageEndpoint - Client endpoint not running"))
+		UE_LOG(LogStormSyncClient, Warning, TEXT("FStormSyncTransportClientModule::GetClientMessageEndpoint - Client endpoint not running"));
 		return nullptr;
 	}
 
@@ -81,13 +81,13 @@ void FStormSyncTransportClientModule::SynchronizePackages(const FStormSyncPackag
 		const TArray<FStormSyncFileDependency> FileDependencies = Result.Get();
 		if (FileDependencies.IsEmpty())
 		{
-			STORM_SYNC_CLIENT_LOG(Error, TEXT("FStormSyncTransportClientModule::SynchronizePackages - Async FileDependencies is empty, something went wrong"))
+			UE_LOG(LogStormSyncClient, Error, TEXT("FStormSyncTransportClientModule::SynchronizePackages - Async FileDependencies is empty, something went wrong"));
 			return;
 		}
 
 		if (!MessageEndpoint.IsValid())
 		{
-			STORM_SYNC_CLIENT_LOG(Error, TEXT("FStormSyncTransportClientModule::SynchronizePackages - Unable to get client message endpoint"))
+			UE_LOG(LogStormSyncClient, Error, TEXT("FStormSyncTransportClientModule::SynchronizePackages - Unable to get client message endpoint"));
 			return;
 		}
 
@@ -95,27 +95,27 @@ void FStormSyncTransportClientModule::SynchronizePackages(const FStormSyncPackag
 		TUniquePtr<FStormSyncTransportSyncRequest> SyncRequestMessage(FMessageEndpoint::MakeMessage<FStormSyncTransportSyncRequest>(PackageNames, PackageDescriptor));
 		if (!SyncRequestMessage.IsValid())
 		{
-			STORM_SYNC_CLIENT_LOG(Error, TEXT("FStormSyncTransportClientModule::SynchronizePackages - Push request message is invalid"))
+			UE_LOG(LogStormSyncClient, Error, TEXT("FStormSyncTransportClientModule::SynchronizePackages - Push request message is invalid"));
 			return;
 		}
 
 		SyncRequestMessage->PackageDescriptor.Dependencies = FileDependencies;
 
-		STORM_SYNC_CLIENT_LOG(Display, TEXT("FStormSyncTransportClientModule::SynchronizePackages - FileDependencies: %d"), FileDependencies.Num())
-		STORM_SYNC_CLIENT_LOG(Display, TEXT("FStormSyncTransportClientModule::SynchronizePackages - Message: %s"), *SyncRequestMessage->ToString())
-		STORM_SYNC_CLIENT_LOG(Display, TEXT("FStormSyncTransportClientModule::SynchronizePackages - Syncing package descriptor %s"), *SyncRequestMessage->PackageDescriptor.ToString())
+		UE_LOG(LogStormSyncClient, Display, TEXT("FStormSyncTransportClientModule::SynchronizePackages - FileDependencies: %d"), FileDependencies.Num());
+		UE_LOG(LogStormSyncClient, Display, TEXT("FStormSyncTransportClientModule::SynchronizePackages - Message: %s"), *SyncRequestMessage->ToString());
+		UE_LOG(LogStormSyncClient, Display, TEXT("FStormSyncTransportClientModule::SynchronizePackages - Syncing package descriptor %s"), *SyncRequestMessage->PackageDescriptor.ToString());
 		MessageEndpoint->Publish(SyncRequestMessage.Release());
 	});
 }
 
 void FStormSyncTransportClientModule::PushPackages(const FStormSyncPackageDescriptor& InPackageDescriptor, const TArray<FName>& InPackageNames, const FMessageAddress& InMessageAddress, const FOnStormSyncPushComplete& InDoneDelegate) const
 {
-	STORM_SYNC_CLIENT_LOG(Display, TEXT("FStormSyncTransportClientModule::PushPackages - PackageDescriptor: %s, InPackageNames: %d, MessageAddressId: %s"), *InPackageDescriptor.ToString(), InPackageNames.Num(), *InMessageAddress.ToString())
+	UE_LOG(LogStormSyncClient, Display, TEXT("FStormSyncTransportClientModule::PushPackages - PackageDescriptor: %s, InPackageNames: %d, MessageAddressId: %s"), *InPackageDescriptor.ToString(), InPackageNames.Num(), *InMessageAddress.ToString());
 
 	const FMessageEndpointSharedPtr MessageEndpoint = GetClientMessageEndpoint();
 	if (!ClientEndpoint.IsValid())
 	{
-		STORM_SYNC_CLIENT_LOG(Error, TEXT("FStormSyncTransportClientModule::PushPackages - Unable to get client endpoint"))
+		UE_LOG(LogStormSyncClient, Error, TEXT("FStormSyncTransportClientModule::PushPackages - Unable to get client endpoint"));
 		return;
 	}
 
@@ -124,13 +124,13 @@ void FStormSyncTransportClientModule::PushPackages(const FStormSyncPackageDescri
 
 void FStormSyncTransportClientModule::PullPackages(const FStormSyncPackageDescriptor& InPackageDescriptor, const TArray<FName>& InPackageNames, const FMessageAddress& InMessageAddress, const FOnStormSyncPullComplete& InDoneDelegate) const
 {
-	STORM_SYNC_CLIENT_LOG(Display, TEXT("FStormSyncTransportClientModule::PullPackages - PackageDescriptor: %s, InPackageNames: %d, MessageAddressId: %s"), *InPackageDescriptor.ToString(), InPackageNames.Num(), *InMessageAddress.ToString())
+	UE_LOG(LogStormSyncClient, Display, TEXT("FStormSyncTransportClientModule::PullPackages - PackageDescriptor: %s, InPackageNames: %d, MessageAddressId: %s"), *InPackageDescriptor.ToString(), InPackageNames.Num(), *InMessageAddress.ToString());
 
 	const FMessageEndpointSharedPtr MessageEndpoint = GetClientMessageEndpoint();
 
 	if (!ClientEndpoint.IsValid())
 	{
-		STORM_SYNC_CLIENT_LOG(Error, TEXT("FStormSyncTransportClientModule::PullPackages - Unable to get client endpoint"))
+		UE_LOG(LogStormSyncClient, Error, TEXT("FStormSyncTransportClientModule::PullPackages - Unable to get client endpoint"));
 		return;
 	}
 
@@ -139,12 +139,12 @@ void FStormSyncTransportClientModule::PullPackages(const FStormSyncPackageDescri
 
 void FStormSyncTransportClientModule::RequestPackagesStatus(const FMessageAddress& InRemoteAddress, const TArray<FName>& InPackageNames, const FOnStormSyncRequestStatusComplete& InDoneDelegate) const
 {
-	STORM_SYNC_CLIENT_LOG(Display, TEXT("FStormSyncTransportClientModule::RequestPackagesStatus - InRemoteAddress: %s, InPackageNames: %d"), *InRemoteAddress.ToString(), InPackageNames.Num())
+	UE_LOG(LogStormSyncClient, Display, TEXT("FStormSyncTransportClientModule::RequestPackagesStatus - InRemoteAddress: %s, InPackageNames: %d"), *InRemoteAddress.ToString(), InPackageNames.Num());
 
 	const FMessageEndpointSharedPtr MessageEndpoint = GetClientMessageEndpoint();
 	if (!ClientEndpoint.IsValid())
 	{
-		STORM_SYNC_CLIENT_LOG(Error, TEXT("FStormSyncTransportClientModule::RequestPackagesStatus - Unable to get client endpoint"))
+		UE_LOG(LogStormSyncClient, Error, TEXT("FStormSyncTransportClientModule::RequestPackagesStatus - Unable to get client endpoint"));
 		return;
 	}
 
@@ -174,7 +174,7 @@ void FStormSyncTransportClientModule::RegisterConsoleCommands()
 		FConsoleCommandWithArgsDelegate::CreateLambda([this](const TArray<FString>& Args)
 		{
 			const FString AddressId = GetClientEndpointMessageAddressId();
-			STORM_SYNC_CLIENT_LOG(Display, TEXT("StormSync.Client.Debug - EndpointId: %s"), *AddressId)
+			UE_LOG(LogStormSyncClient, Display, TEXT("StormSync.Client.Debug - EndpointId: %s"), *AddressId);
 		}),
 		ECVF_Default
 	));
@@ -190,11 +190,11 @@ void FStormSyncTransportClientModule::RegisterConsoleCommands()
 				const TSharedPtr<FMessageEndpoint, ESPMode::ThreadSafe> MessageEndpoint = ClientEndpoint->GetMessageEndpoint();
 				if (!MessageEndpoint.IsValid())
 				{
-					STORM_SYNC_CLIENT_LOG(Error, TEXT("StormSync.Client.Debug.Ping - Unable to send Connect Message cause Message Endpoint is invalid"))
+					UE_LOG(LogStormSyncClient, Error, TEXT("StormSync.Client.Debug.Ping - Unable to send Connect Message cause Message Endpoint is invalid"));
 					return;
 				}
 
-				STORM_SYNC_CLIENT_LOG(Display, TEXT("StormSync.Client.Debug.Ping - Publish Client Connect Message ..."))
+				UE_LOG(LogStormSyncClient, Display, TEXT("StormSync.Client.Debug.Ping - Publish Client Connect Message ..."));
 				FStormSyncTransportStatusPing* Message = FMessageEndpoint::MakeMessage<FStormSyncTransportStatusPing>();
 				MessageEndpoint->Publish(Message);
 			}
@@ -208,7 +208,7 @@ void FStormSyncTransportClientModule::RegisterConsoleCommands()
 		FConsoleCommandWithArgsDelegate::CreateLambda([this](const TArray<FString>& Args)
 		{
 			const FString InstanceId = FApp::GetInstanceId().ToString();
-			STORM_SYNC_CLIENT_LOG(Display, TEXT("StormSync.Client.Debug - InstanceId: %s"), *InstanceId)
+			UE_LOG(LogStormSyncClient, Display, TEXT("StormSync.Client.Debug - InstanceId: %s"), *InstanceId);
 		}),
 		ECVF_Default
 	));
@@ -231,12 +231,12 @@ void FStormSyncTransportClientModule::OnEngineLoopInitComplete()
 		const TSharedPtr<FMessageEndpoint, ESPMode::ThreadSafe> MessageEndpoint = ClientEndpoint->GetMessageEndpoint();
 		if (!MessageEndpoint.IsValid())
 		{
-			STORM_SYNC_CLIENT_LOG(Error, TEXT("FStormSyncTransportClientModule::OnEngineLoopInitComplete - Unable to send Connect Message cause Message Endpoint is invalid"))
+			UE_LOG(LogStormSyncClient, Error, TEXT("FStormSyncTransportClientModule::OnEngineLoopInitComplete - Unable to send Connect Message cause Message Endpoint is invalid"));
 			return;
 		}
 
 		// We broadcast a message to notify others about this editor instance (this is required so that further "direct" send are received on the other end)
-		STORM_SYNC_CLIENT_LOG(Display, TEXT("FStormSyncTransportClientModule::OnEngineLoopInitComplete - Publish Client Connect Message ..."))
+		UE_LOG(LogStormSyncClient, Display, TEXT("FStormSyncTransportClientModule::OnEngineLoopInitComplete - Publish Client Connect Message ..."));
 		FStormSyncTransportStatusPing* Message = FMessageEndpoint::MakeMessage<FStormSyncTransportStatusPing>();
 		MessageEndpoint->Publish(Message);
 	}
@@ -245,12 +245,12 @@ void FStormSyncTransportClientModule::OnEngineLoopInitComplete()
 void FStormSyncTransportClientModule::ExecutePing(const TArray<FString>& Args)
 {
 	const FString Argv = FString::Join(Args, TEXT(""));
-	STORM_SYNC_CLIENT_LOG(Display, TEXT("FStormSyncTransportClientModule::ExecutePing - %s"), *Argv)
+	UE_LOG(LogStormSyncClient, Display, TEXT("FStormSyncTransportClientModule::ExecutePing - %s"), *Argv);
 
 	const FMessageEndpointSharedPtr MessageEndpoint = GetClientMessageEndpoint();
 	if (!MessageEndpoint.IsValid())
 	{
-		STORM_SYNC_CLIENT_LOG(Warning, TEXT("FStormSyncTransportClientModule::ExecutePing - Unable to get client message endpoint"))
+		UE_LOG(LogStormSyncClient, Warning, TEXT("FStormSyncTransportClientModule::ExecutePing - Unable to get client message endpoint"));
 		return;
 	}
 
@@ -263,7 +263,7 @@ FStormSyncPackageDescriptor FStormSyncTransportClientModule::CreatePackageDescri
 	if (!FParse::Value(*Argv, TEXT("-name="), PackageDescriptor.Name))
 	{
 		PackageDescriptor.Name = DefaultPakName;
-		STORM_SYNC_CLIENT_LOG(Display, TEXT("FStormSyncEditorModule::CreatePackageDescriptorFromCommandLine - Missing -name parameter, using default \"%s\""), *PackageDescriptor.Name);
+		UE_LOG(LogStormSyncClient, Display, TEXT("FStormSyncEditorModule::CreatePackageDescriptorFromCommandLine - Missing -name parameter, using default \"%s\""), *PackageDescriptor.Name);
 	}
 
 	FParse::Value(*Argv, TEXT("-version="), PackageDescriptor.Version);
@@ -276,7 +276,7 @@ FStormSyncPackageDescriptor FStormSyncTransportClientModule::CreatePackageDescri
 void FStormSyncTransportClientModule::ExecuteSyncPak(const TArray<FString>& Args)
 {
 	const FString Argv = FString::Join(Args, TEXT(" "));
-	STORM_SYNC_CLIENT_LOG(Display, TEXT("FStormSyncTransportClientModule::ExecuteSyncPak - Argv: %s"), *Argv)
+	UE_LOG(LogStormSyncClient, Display, TEXT("FStormSyncTransportClientModule::ExecuteSyncPak - Argv: %s"), *Argv);
 
 	// Parse command line.
 	TArray<FName> PackageNames;
@@ -284,7 +284,7 @@ void FStormSyncTransportClientModule::ExecuteSyncPak(const TArray<FString>& Args
 
 	if (PackageNames.IsEmpty())
 	{
-		STORM_SYNC_CLIENT_LOG(Error, TEXT("FStormSyncTransportClientModule::ExecuteSyncPak - Missing at least one package name to sync."))
+		UE_LOG(LogStormSyncClient, Error, TEXT("FStormSyncTransportClientModule::ExecuteSyncPak - Missing at least one package name to sync."));
 		return;
 	}
 
diff --git a/Engine/Plugins/Experimental/StormSync/Source/StormSyncTransportCore/Private/StormSyncTransportCoreLog.h b/Engine/Plugins/Experimental/StormSync/Source/StormSyncTransportCore/Private/StormSyncTransportCoreLog.h
--- a/Engine/Plugins/Experimental/StormSync/Source/StormSyncTransportCore/Private/StormSyncTransportCoreLog.h
+++ b/Engine/Plugins/Experimental/StormSync/Source/StormSyncTransportCore/Private/StormSyncTransportCoreLog.h
@@ -5,8 +5,3 @@
 #include "Logging/LogMacros.h"
 
 DECLARE_LOG_CATEGORY_EXTERN(LogStormSyncTransportCore, Display, All);
-
-#define STORM_SYNC_CORE_LOG(Verbosity, Format, ...) \
-{ \
-    UE_LOG(LogStormSyncTransportCore, Verbosity, Format, ##__VA_ARGS__); \
-}
diff --git a/Engine/Plugins/Experimental/StormSync/Source/StormSyncTransportCore/Private/StormSyncTransportNetworkUtils.cpp b/Engine/Plugins/Experimental/StormSync/Source/StormSyncTransportCore/Private/StormSyncTransportNetworkUtils.cpp
--- a/Engine/Plugins/Experimental/StormSync/Source/StormSyncTransportCore/Private/StormSyncTransportNetworkUtils.cpp
+++ b/Engine/Plugins/Experimental/StormSync/Source/StormSyncTransportCore/Private/StormSyncTransportNetworkUtils.cpp
@@ -26,7 +26,7 @@ FString FStormSyncTransportNetworkUtils::GetTcpEndpointAddress()
 	FIPv4Endpoint Endpoint;
 	if (!FIPv4Endpoint::Parse(ServerEndpoint, Endpoint))
 	{
-		STORM_SYNC_CORE_LOG(Error, TEXT("UStormSyncTransportServerUtils::GetTcpEndpointAddress - Failed to parse endpoint '%s'"), *ServerEndpoint);
+		UE_LOG(LogStormSyncTransportCore, Error, TEXT("UStormSyncTransportServerUtils::GetTcpEndpointAddress - Failed to parse endpoint '%s'"), *ServerEndpoint);
 		return TEXT("");
 	}
 
@@ -45,7 +45,7 @@ TArray<FString> FStormSyncTransportNetworkUtils::GetLocalAdapterAddresses()
 	FIPv4Endpoint Endpoint;
 	if (!FIPv4Endpoint::Parse(ServerEndpoint, Endpoint))
 	{
-		STORM_SYNC_CORE_LOG(Error, TEXT("UStormSyncTransportServerUtils::GetLocalAdapterAddresses - Failed to parse endpoint '%s'"), *ServerEndpoint);
+		UE_LOG(LogStormSyncTransportCore, Error, TEXT("UStormSyncTransportServerUtils::GetLocalAdapterAddresses - Failed to parse endpoint '%s'"), *ServerEndpoint);
 		return AdapterAddresses;
 	}
 	
diff --git a/Engine/Plugins/Experimental/StormSync/Source/StormSyncTransportCore/Private/StormSyncTransportSettings.cpp b/Engine/Plugins/Experimental/StormSync/Source/StormSyncTransportCore/Private/StormSyncTransportSettings.cpp
--- a/Engine/Plugins/Experimental/StormSync/Source/StormSyncTransportCore/Private/StormSyncTransportSettings.cpp
+++ b/Engine/Plugins/Experimental/StormSync/Source/StormSyncTransportCore/Private/StormSyncTransportSettings.cpp
@@ -54,7 +54,7 @@ FString UStormSyncTransportSettings::GetServerEndpoint() const
 		if (!GetUdpMessagingUnicastEndpoint(UnicastHostname, UnicastPort))
 		{
 			// Either unable to get the setting or malformed
-			STORM_SYNC_CORE_LOG(Display, TEXT("UStormSyncTransportSettings::GetServerEndpoint - Was unable to get UnicastEndpoint from UDPMessagingSettings"));
+			UE_LOG(LogStormSyncTransportCore, Display, TEXT("UStormSyncTransportSettings::GetServerEndpoint - Was unable to get UnicastEndpoint from UDPMessagingSettings"));
 			return ServerEndpoint;
 		}
 
diff --git a/Engine/Plugins/Experimental/StormSync/Source/StormSyncTransportServer/Private/ServiceDiscovery/StormSyncDiscoveryManager.cpp b/Engine/Plugins/Experimental/StormSync/Source/StormSyncTransportServer/Private/ServiceDiscovery/StormSyncDiscoveryManager.cpp
--- a/Engine/Plugins/Experimental/StormSync/Source/StormSyncTransportServer/Private/ServiceDiscovery/StormSyncDiscoveryManager.cpp
+++ b/Engine/Plugins/Experimental/StormSync/Source/StormSyncTransportServer/Private/ServiceDiscovery/StormSyncDiscoveryManager.cpp
@@ -28,10 +28,10 @@ FStormSyncDiscoveryManager::FStormSyncDiscoveryManager(double InHeartbeatTimeout
 	bRunning = MessageEndpoint.IsValid();
 	if (bRunning)
 	{
-		STORM_SYNC_SERVER_LOG(Display, TEXT("FStormSyncDiscoveryManager::FStormSyncDiscoveryManager - Subscribe to messages"))
+		UE_LOG(LogStormSyncServer, Display, TEXT("FStormSyncDiscoveryManager::FStormSyncDiscoveryManager - Subscribe to messages"));
 		MessageEndpoint->Subscribe<FStormSyncTransportConnectMessage>();
 
-		STORM_SYNC_SERVER_LOG(Display, TEXT("FStormSyncDiscoveryManager::FStormSyncDiscoveryManager - Start Thread"))
+		UE_LOG(LogStormSyncServer, Display, TEXT("FStormSyncDiscoveryManager::FStormSyncDiscoveryManager - Start Thread"));
 		Thread = FRunnableThread::Create(this, TEXT("StormSyncDiscoveryManager"));
 	}
 }
@@ -108,7 +108,7 @@ uint32 FStormSyncDiscoveryManager::Run()
 			// Handle inactive addresses
 			for (FMessageAddress MessageAddress : DisconnectedAddresses)
 			{
-				STORM_SYNC_SERVER_LOG(Display, TEXT("FStormSyncDiscoveryManager::Run - %s became invalid"), *MessageAddress.ToString())
+				UE_LOG(LogStormSyncServer, Display, TEXT("FStormSyncDiscoveryManager::Run - %s became invalid"), *MessageAddress.ToString());
 
 				// Stop Heartbeat for this recipient
 				FStormSyncHeartbeatEmitter& HeartbeatEmitter = IStormSyncTransportServerModule::Get().GetHeartbeatEmitter();
@@ -150,11 +150,11 @@ void FStormSyncDiscoveryManager::PublishConnectMessage()
 {
 	if (!MessageEndpoint.IsValid())
 	{
-		STORM_SYNC_SERVER_LOG(Error, TEXT("FStormSyncDiscoveryManager::PublishConnectMessage - Unable to send Connect Message cause Message Endpoint is invalid"))
+		UE_LOG(LogStormSyncServer, Error, TEXT("FStormSyncDiscoveryManager::PublishConnectMessage - Unable to send Connect Message cause Message Endpoint is invalid"));
 		return;
 	}
 
-	STORM_SYNC_SERVER_LOG(VeryVerbose, TEXT("FStormSyncTransportClientModule::PublishConnectMessage - Publish Connect Message ..."))
+	UE_LOG(LogStormSyncServer, VeryVerbose, TEXT("FStormSyncTransportClientModule::PublishConnectMessage - Publish Connect Message ..."));
 	TUniquePtr<FStormSyncTransportConnectMessage> Message(FMessageEndpoint::MakeMessage<FStormSyncTransportConnectMessage>());
 	if (Message.IsValid())
 	{
@@ -191,11 +191,11 @@ void FStormSyncDiscoveryManager::BroadcastCoreDelegatesFromQueue(const TArray<FS
 void FStormSyncDiscoveryManager::HandleConnectMessage(const FStormSyncTransportConnectMessage& InMessage, const TSharedRef<IMessageContext, ESPMode::ThreadSafe>& MessageContext)
 {
 	const FMessageAddress SenderMessageAddress = MessageContext->GetSender();
-	STORM_SYNC_SERVER_LOG(VeryVerbose, TEXT("FStormSyncDiscoveryManager::HandleConnectMessage - Received connect Message from %s: %s"), *SenderMessageAddress.ToString(), *InMessage.ToString())
+	UE_LOG(LogStormSyncServer, VeryVerbose, TEXT("FStormSyncDiscoveryManager::HandleConnectMessage - Received connect Message from %s: %s"), *SenderMessageAddress.ToString(), *InMessage.ToString());
 
 	const FMessageAddress MessageEndpointAddress = MessageEndpoint->GetAddress();
 
-	STORM_SYNC_SERVER_LOG(VeryVerbose, TEXT("FStormSyncDiscoveryManager::HandleConnectMessage - SenderMessageAddress: %s, MessageEndpointAddress: %s (Same: %s)"), *SenderMessageAddress.ToString(), *MessageEndpointAddress.ToString(), SenderMessageAddress == MessageEndpointAddress ? TEXT("true") : TEXT("false"))
+	UE_LOG(LogStormSyncServer, VeryVerbose, TEXT("FStormSyncDiscoveryManager::HandleConnectMessage - SenderMessageAddress: %s, MessageEndpointAddress: %s (Same: %s)"), *SenderMessageAddress.ToString(), *MessageEndpointAddress.ToString(), SenderMessageAddress == MessageEndpointAddress ? TEXT("true") : TEXT("false"));
 
 	if (SenderMessageAddress != MessageEndpointAddress && !ConnectedAddresses.Contains(SenderMessageAddress))
 	{
@@ -223,7 +223,7 @@ void FStormSyncDiscoveryManager::RegisterConnection(const FStormSyncTransportCon
 	HeartbeatEmitter.StartHeartbeat(InMessageAddress, MessageEndpoint);
 
 	// Send back connect message so that this recipient knows about this editor instance
-	STORM_SYNC_SERVER_LOG(Display, TEXT("FStormSyncDiscoveryManager::HandleConnectMessage - Send Connect Message to %s..."), *InMessageAddress.ToString())
+	UE_LOG(LogStormSyncServer, Display, TEXT("FStormSyncDiscoveryManager::HandleConnectMessage - Send Connect Message to %s..."), *InMessageAddress.ToString());
 	TUniquePtr<FStormSyncTransportConnectMessage> Message(FMessageEndpoint::MakeMessage<FStormSyncTransportConnectMessage>());
 	if (Message.IsValid())
 	{
@@ -251,7 +251,7 @@ void FStormSyncDiscoveryManager::HandleHeartbeatMessage(const FStormSyncTranspor
 {
 	const FMessageAddress MessageAddress = MessageContext->GetSender();
 
-	STORM_SYNC_SERVER_LOG(VeryVerbose, TEXT("FStormSyncDiscoveryManager::HandleHeartbeatMessage - Received Heartbeat Message from %s"), *MessageAddress.ToString())
+	UE_LOG(LogStormSyncServer, VeryVerbose, TEXT("FStormSyncDiscoveryManager::HandleHeartbeatMessage - Received Heartbeat Message from %s"), *MessageAddress.ToString());
 	
 	{
 		// Handle revive connection in case we receive heartbeats again, most likely meaning
diff --git a/Engine/Plugins/Experimental/StormSync/Source/StormSyncTransportServer/Private/ServiceDiscovery/StormSyncHeartbeatEmitter.cpp b/Engine/Plugins/Experimental/StormSync/Source/StormSyncTransportServer/Private/ServiceDiscovery/StormSyncHeartbeatEmitter.cpp
--- a/Engine/Plugins/Experimental/StormSync/Source/StormSyncTransportServer/Private/ServiceDiscovery/StormSyncHeartbeatEmitter.cpp
+++ b/Engine/Plugins/Experimental/StormSync/Source/StormSyncTransportServer/Private/ServiceDiscovery/StormSyncHeartbeatEmitter.cpp
@@ -22,7 +22,7 @@ FStormSyncHeartbeatEmitter::FStormSyncHeartbeatEmitter()
 
 void FStormSyncHeartbeatEmitter::StartHeartbeat(const FMessageAddress& RecipientAddress, const TSharedPtr<FMessageEndpoint, ESPMode::ThreadSafe>& MessageEndpoint)
 {
-	STORM_SYNC_SERVER_LOG(Verbose, TEXT("Start Heartbeat for %s"), *RecipientAddress.ToString())
+	UE_LOG(LogStormSyncServer, Verbose, TEXT("Start Heartbeat for %s"), *RecipientAddress.ToString());
 
 	if (!bIsRunning)
 	{
@@ -58,7 +58,7 @@ uint32 FStormSyncHeartbeatEmitter::Run()
 			{
 				if (const TSharedPtr<FMessageEndpoint, ESPMode::ThreadSafe> MessageEndpoint = Recipient.MessageEndpoint.Pin())
 				{
-					STORM_SYNC_SERVER_LOG(VeryVerbose, TEXT("Sending Heartbeat Message to %s"), *Recipient.ConnectionAddress.ToString())
+					UE_LOG(LogStormSyncServer, VeryVerbose, TEXT("Sending Heartbeat Message to %s"), *Recipient.ConnectionAddress.ToString());
 					if (FStormSyncTransportHeartbeatMessage* Message = FMessageEndpoint::MakeMessage<FStormSyncTransportHeartbeatMessage>())
 					{
 						Message->bIsServerRunning = false;
diff --git a/Engine/Plugins/Experimental/StormSync/Source/StormSyncTransportServer/Private/Socket/StormSyncTransportTcpServer.cpp b/Engine/Plugins/Experimental/StormSync/Source/StormSyncTransportServer/Private/Socket/StormSyncTransportTcpServer.cpp
--- a/Engine/Plugins/Experimental/StormSync/Source/StormSyncTransportServer/Private/Socket/StormSyncTransportTcpServer.cpp
+++ b/Engine/Plugins/Experimental/StormSync/Source/StormSyncTransportServer/Private/Socket/StormSyncTransportTcpServer.cpp
@@ -53,7 +53,7 @@ uint32 FStormSyncTransportTcpServer::Run()
 
 		if (DeltaTime > 0.1)
 		{
-			STORM_SYNC_SERVER_LOG(Warning, TEXT("Hitch detected; %.3f seconds since prior tick"), DeltaTime);
+			UE_LOG(LogStormSyncServer, Warning, TEXT("Hitch detected; %.3f seconds since prior tick"), DeltaTime);
 		}
 
 		bListenerIsRunning = Tick();
@@ -107,7 +107,7 @@ bool FStormSyncTransportTcpServer::Tick()
 			int32 BytesRead = 0;
 			if (!ClientSocket->Recv(Buffer.GetData(), PendingDataSize, BytesRead, ESocketReceiveFlags::None))
 			{
-				STORM_SYNC_SERVER_LOG(Error, TEXT("Error while receiving data via endpoint %s"), *ClientEndpoint.ToString());
+				UE_LOG(LogStormSyncServer, Error, TEXT("Error while receiving data via endpoint %s"), *ClientEndpoint.ToString());
 				continue;
 			}
 
@@ -125,7 +125,7 @@ bool FStormSyncTransportTcpServer::Tick()
 				{
 					const uint32 Size = *reinterpret_cast<const uint32*>(MessageBuffer.GetData());
 			
-					STORM_SYNC_SERVER_LOG(Verbose, TEXT("Received size: %d. Setting it for %s"), Size, *ClientEndpoint.ToString());
+					UE_LOG(LogStormSyncServer, Verbose, TEXT("Received size: %d. Setting it for %s"), Size, *ClientEndpoint.ToString());
 					Connection.Value->BufferExpectedSize = Size;
 					Connection.Value->BufferStartTime = FPlatformTime::Seconds();
 					bHasReceivedSize = true;
@@ -158,11 +158,11 @@ bool FStormSyncTransportTcpServer::StartListening()
 	if (SocketListener->IsActive())
 	{
 		SocketListener->OnConnectionAccepted().BindRaw(this, &FStormSyncTransportTcpServer::OnIncomingConnection);
-		STORM_SYNC_SERVER_LOG(Display, TEXT("Started listening on %s:%d"), *SocketListener->GetLocalEndpoint().Address.ToString(), SocketListener->GetLocalEndpoint().Port);
+		UE_LOG(LogStormSyncServer, Display, TEXT("Started listening on %s:%d"), *SocketListener->GetLocalEndpoint().Address.ToString(), SocketListener->GetLocalEndpoint().Port);
 		return true;
 	}
 
-	STORM_SYNC_SERVER_LOG(Error, TEXT("Could not create Tcp Listener!"));
+	UE_LOG(LogStormSyncServer, Error, TEXT("Could not create Tcp Listener!"));
 	return false;
 }
 
@@ -172,7 +172,7 @@ void FStormSyncTransportTcpServer::StopListening()
 
 	if (SocketListener.IsValid())
 	{
-		STORM_SYNC_SERVER_LOG(Display, TEXT("No longer listening on %s:%d"), *SocketListener->GetLocalEndpoint().Address.ToString(), SocketListener->GetLocalEndpoint().Port);
+		UE_LOG(LogStormSyncServer, Display, TEXT("No longer listening on %s:%d"), *SocketListener->GetLocalEndpoint().Address.ToString(), SocketListener->GetLocalEndpoint().Port);
 		SocketListener.Reset();
 	}
 
@@ -201,7 +201,7 @@ bool FStormSyncTransportTcpServer::IsActive() const
 
 bool FStormSyncTransportTcpServer::OnIncomingConnection(FSocket* InSocket, const FIPv4Endpoint& InEndpoint)
 {
-	STORM_SYNC_SERVER_LOG(Display, TEXT("Incoming connection via %s:%d"), *InEndpoint.Address.ToString(), InEndpoint.Port);
+	UE_LOG(LogStormSyncServer, Display, TEXT("Incoming connection via %s:%d"), *InEndpoint.Address.ToString(), InEndpoint.Port);
 
 	InSocket->SetNoDelay(true);
 	PendingConnections.Enqueue(TPair<FIPv4Endpoint, TSharedPtr<FSocket>>(InEndpoint, MakeShareable(InSocket)));
@@ -212,7 +212,7 @@ bool FStormSyncTransportTcpServer::OnIncomingConnection(FSocket* InSocket, const
 void FStormSyncTransportTcpServer::HandleIncomingBuffer(const FIPv4Endpoint& InEndpoint, const TSharedPtr<FSocket>& InClientSocket, const TArray<uint8>& InBytes)
 {
 	const int32 ReceivedBufferSize = InBytes.Num();
-	STORM_SYNC_SERVER_LOG(Verbose, TEXT("\tHandle Incoming buffer via %s:%d (Total Buffer Size: %d)"), *InEndpoint.Address.ToString(), InEndpoint.Port, ReceivedBufferSize);
+	UE_LOG(LogStormSyncServer, Verbose, TEXT("\tHandle Incoming buffer via %s:%d (Total Buffer Size: %d)"), *InEndpoint.Address.ToString(), InEndpoint.Port, ReceivedBufferSize);
 
 	// Send back to client the size buffer we received so far
 	SendMessage(CreateMessage(FStormSyncTransportTcpSizePacket(ReceivedBufferSize)), InEndpoint);
@@ -222,7 +222,7 @@ void FStormSyncTransportTcpServer::HandleIncomingBuffer(const FIPv4Endpoint& InE
 		const uint32 BufferExpectedSize = ClientConnections.FindChecked(InEndpoint)->BufferExpectedSize;
 		if (ReceivedBufferSize == BufferExpectedSize)
 		{
-			STORM_SYNC_SERVER_LOG(Display, TEXT("\tWe received full buffer!!"));
+			UE_LOG(LogStormSyncServer, Display, TEXT("\tWe received full buffer!!"));
 			
 			// Send back to client transfer complete packet, so that it knows we're done with the transfer
 			SendMessage(CreateMessage(FStormSyncTransportTcpTransferCompletePacket()), InEndpoint);
@@ -236,7 +236,7 @@ void FStormSyncTransportTcpServer::HandleIncomingBuffer(const FIPv4Endpoint& InE
 
 void FStormSyncTransportTcpServer::HandleReceivedBuffer(const FIPv4Endpoint& InEndpoint, const TSharedPtr<FSocket>& InClientSocket, const TArray<uint8>& InBytes)
 {
-	STORM_SYNC_SERVER_LOG(Display, TEXT("Handle received buffer via %s:%d (Buffer Size: %d)"), *InEndpoint.Address.ToString(), InEndpoint.Port, InBytes.Num());
+	UE_LOG(LogStormSyncServer, Display, TEXT("Handle received buffer via %s:%d (Buffer Size: %d)"), *InEndpoint.Address.ToString(), InEndpoint.Port, InBytes.Num());
 	ReceivedBufferEvent.Broadcast(InEndpoint, InClientSocket, MakeShared<TArray<uint8>>(InBytes));
 
 	// Fully received buffer, clean up connection and the state associated with it
@@ -255,7 +255,7 @@ void FStormSyncTransportTcpServer::CleanUpDisconnectedSockets()
 		const float ClientTimeout = Connection.Value->InactiveTimeout;
 		if (CurrentTime - Connection.Value->LastActivityTime > ClientTimeout)
 		{
-			STORM_SYNC_SERVER_LOG(Warning, TEXT("Client %s has been inactive for more than %.1fs -- closing connection"), *RemoteEndpoint.ToString(), ClientTimeout);
+			UE_LOG(LogStormSyncServer, Warning, TEXT("Client %s has been inactive for more than %.1fs -- closing connection"), *RemoteEndpoint.ToString(), ClientTimeout);
 			DisconnectedClients.Add(RemoteEndpoint);
 		}
 	}
@@ -277,7 +277,7 @@ void FStormSyncTransportTcpServer::CleanUpDisconnectedSockets()
 
 void FStormSyncTransportTcpServer::DisconnectClient(const FIPv4Endpoint& InClientEndpoint)
 {
-	STORM_SYNC_SERVER_LOG(Display, TEXT("Client %s disconnected"), *InClientEndpoint.ToString());
+	UE_LOG(LogStormSyncServer, Display, TEXT("Client %s disconnected"), *InClientEndpoint.ToString());
 
 	ClientConnections.Remove(InClientEndpoint);
 }
@@ -292,11 +292,11 @@ bool FStormSyncTransportTcpServer::SendMessage(const FString& InMessage, const F
 			return false;
 		}
 
-		STORM_SYNC_SERVER_LOG(Verbose, TEXT("Sending message to %s - Message: %s"), *InEndpoint.ToString(), *InMessage);
+		UE_LOG(LogStormSyncServer, Verbose, TEXT("Sending message to %s - Message: %s"), *InEndpoint.ToString(), *InMessage);
 		int32 BytesSent = 0;
 		return ClientSocket->Send(reinterpret_cast<uint8*>(TCHAR_TO_UTF8(*InMessage)), InMessage.Len() + 1, BytesSent);
 	}
 
-	STORM_SYNC_SERVER_LOG(Verbose, TEXT("Trying to send message to disconnected client %s - Message: %s"), *InEndpoint.ToString(), *InMessage);
+	UE_LOG(LogStormSyncServer, Verbose, TEXT("Trying to send message to disconnected client %s - Message: %s"), *InEndpoint.ToString(), *InMessage);
 	return false;
 }
diff --git a/Engine/Plugins/Experimental/StormSync/Source/StormSyncTransportServer/Private/StormSyncTransportServerEndpoint.cpp b/Engine/Plugins/Experimental/StormSync/Source/StormSyncTransportServer/Private/StormSyncTransportServerEndpoint.cpp
--- a/Engine/Plugins/Experimental/StormSync/Source/StormSyncTransportServer/Private/StormSyncTransportServerEndpoint.cpp
+++ b/Engine/Plugins/Experimental/StormSync/Source/StormSyncTransportServer/Private/StormSyncTransportServerEndpoint.cpp
@@ -35,7 +35,7 @@ bool FStormSyncTransportServerEndpoint::StartTcpListener(const FIPv4Endpoint& In
 {
 	const UStormSyncTransportSettings* Settings = GetDefault<UStormSyncTransportSettings>();
 	
-	STORM_SYNC_SERVER_LOG(Display, TEXT("FStormSyncTransportServerEndpoint::StartTcpListener - Starting TCP listener with endpoint: %s"), *InEndpoint.ToString())
+	UE_LOG(LogStormSyncServer, Display, TEXT("FStormSyncTransportServerEndpoint::StartTcpListener - Starting TCP listener with endpoint: %s"), *InEndpoint.ToString());
 
 	TcpServer = MakeUnique<FStormSyncTransportTcpServer>(InEndpoint.Address, InEndpoint.Port, Settings->GetInactiveTimeoutSeconds());
 	TcpServer->OnReceivedBuffer().AddRaw(this, &FStormSyncTransportServerEndpoint::HandleReceivedTcpBuffer);
@@ -68,7 +68,7 @@ bool FStormSyncTransportServerEndpoint::StartTcpListener(const FIPv4Endpoint& In
 		FSlateNotificationManager::Get().AddNotification(MoveTemp(Info));
 #endif
 		
-		STORM_SYNC_SERVER_LOG(Error, TEXT("FStormSyncTransportServerEndpoint::StartTcpListener - %s"), *ErrorMessage.ToString())
+		UE_LOG(LogStormSyncServer, Error, TEXT("FStormSyncTransportServerEndpoint::StartTcpListener - %s"), *ErrorMessage.ToString());
 		return false;
 	}
 	
@@ -89,7 +89,7 @@ bool FStormSyncTransportServerEndpoint::StartTcpListener()
 		return StartTcpListener(Endpoint);
 	}
 	
-	STORM_SYNC_SERVER_LOG(Error, TEXT("FStormSyncTransportServerEndpoint::StartTcpListener - Failed to parse endpoint '%s'. TCP Socket listener will be disabled!"), *ServerEndpoint);
+	UE_LOG(LogStormSyncServer, Error, TEXT("FStormSyncTransportServerEndpoint::StartTcpListener - Failed to parse endpoint '%s'. TCP Socket listener will be disabled!"), *ServerEndpoint);
 	return false;
 }
 
@@ -125,19 +125,19 @@ void FStormSyncTransportServerEndpoint::InitializeMessaging(const FString& InEnd
 	// Message endpoint already created
 	if (MessageEndpoint.IsValid())
 	{
-		STORM_SYNC_SERVER_LOG(Warning, TEXT("Message endpoint already initialized"))
+		UE_LOG(LogStormSyncServer, Warning, TEXT("Message endpoint already initialized"));
 		return;
 	}
 
 	const IMessageBusPtr MessageBus = MessageBusPtr.Pin();
 	if (!MessageBus.IsValid())
 	{
-		STORM_SYNC_SERVER_LOG(Warning, TEXT("Default message bus is invalid"))
+		UE_LOG(LogStormSyncServer, Warning, TEXT("Default message bus is invalid"));
 		return;
 	}
 
 	const FString MessageEndpointName = FString::Printf(TEXT("StormSync%sEndpoint"), *InEndpointFriendlyName);
-	STORM_SYNC_SERVER_LOG(Verbose, TEXT("FStormSyncTransportServerEndpoint::InitializeMessaging - Setting up %s message endpoint"), *MessageEndpointName)
+	UE_LOG(LogStormSyncServer, Verbose, TEXT("FStormSyncTransportServerEndpoint::InitializeMessaging - Setting up %s message endpoint"), *MessageEndpointName);
 	
 	FMessageEndpointBuilder EndpointBuilder = FMessageEndpoint::Builder(*MessageEndpointName, MessageBus.ToSharedRef())
 		.ReceivingOnThread(ENamedThreads::Type::GameThread)
@@ -150,7 +150,7 @@ void FStormSyncTransportServerEndpoint::InitializeMessaging(const FString& InEnd
 	check(MessageEndpoint.IsValid());
 
 	// Message subscribes
-	STORM_SYNC_SERVER_LOG(Display, TEXT("FStormSyncTransportServerEndpoint::InitializeMessaging - Subscribe to messages"))
+	UE_LOG(LogStormSyncServer, Display, TEXT("FStormSyncTransportServerEndpoint::InitializeMessaging - Subscribe to messages"));
 	MessageEndpoint->Subscribe<FStormSyncTransportPingMessage>();
 	MessageEndpoint->Subscribe<FStormSyncTransportSyncRequest>();
 }
@@ -176,10 +176,10 @@ void FStormSyncTransportServerEndpoint::ShutdownTcpListener()
 
 void FStormSyncTransportServerEndpoint::HandlePingMessage(const FStormSyncTransportPingMessage& InMessage, const TSharedRef<IMessageContext, ESPMode::ThreadSafe>& InMessageContext)
 {
-	STORM_SYNC_SERVER_LOG(Display, TEXT("FStormSyncTransportServerEndpoint::HandlePingMessage - Received ping Message %s"), *InMessage.ToString())
+	UE_LOG(LogStormSyncServer, Display, TEXT("FStormSyncTransportServerEndpoint::HandlePingMessage - Received ping Message %s"), *InMessage.ToString());
 	if (!MessageEndpoint.IsValid())
 	{
-		STORM_SYNC_SERVER_LOG(Display, TEXT("FStormSyncTransportServerEndpoint::HandlePingMessage - Invalid MessageEndpoint, can't pong back"))
+		UE_LOG(LogStormSyncServer, Display, TEXT("FStormSyncTransportServerEndpoint::HandlePingMessage - Invalid MessageEndpoint, can't pong back"));
 		return;
 	}
 
@@ -189,18 +189,18 @@ void FStormSyncTransportServerEndpoint::HandlePingMessage(const FStormSyncTransp
 
 void FStormSyncTransportServerEndpoint::HandlePongMessage(const FStormSyncTransportPongMessage& InMessage, const TSharedRef<IMessageContext, ESPMode::ThreadSafe>& InMessageContext)
 {
-	STORM_SYNC_SERVER_LOG(Display, TEXT("FStormSyncTransportServerEndpoint::HandlePongMessage - Received pong Message %s"), *InMessage.ToString())
+	UE_LOG(LogStormSyncServer, Display, TEXT("FStormSyncTransportServerEndpoint::HandlePongMessage - Received pong Message %s"), *InMessage.ToString());
 }
 
 void FStormSyncTransportServerEndpoint::HandleReceivedTcpBuffer(const FIPv4Endpoint& InEndpoint, const TSharedPtr<FSocket>& InSocket, const FStormSyncBufferPtr& InBuffer)
 {
 	if (!InBuffer.IsValid())
 	{
-		STORM_SYNC_SERVER_LOG(Error, TEXT("Error received tcp buffer via %s:%d (Buffer is invalid)"), *InEndpoint.Address.ToString(), InEndpoint.Port);
+		UE_LOG(LogStormSyncServer, Error, TEXT("Error received tcp buffer via %s:%d (Buffer is invalid)"), *InEndpoint.Address.ToString(), InEndpoint.Port);
 		return;
 	}
 
-	STORM_SYNC_SERVER_LOG(Display, TEXT("Handle received tcp buffer via %s:%d (Buffer Size: %d)"), *InEndpoint.Address.ToString(), InEndpoint.Port, InBuffer->Num());
+	UE_LOG(LogStormSyncServer, Display, TEXT("Handle received tcp buffer via %s:%d (Buffer Size: %d)"), *InEndpoint.Address.ToString(), InEndpoint.Port, InBuffer->Num());
 
 	// Local copy of the endpoint we're receiving a buffer from
 	FIPv4Endpoint Endpoint = InEndpoint;
@@ -220,7 +220,7 @@ void FStormSyncTransportServerEndpoint::HandleReceivedTcpBuffer(const FIPv4Endpo
 
 void FStormSyncTransportServerEndpoint::HandleSyncRequestMessage(const FStormSyncTransportSyncRequest& InMessage, const TSharedRef<IMessageContext, ESPMode::ThreadSafe>& InMessageContext)
 {
-	STORM_SYNC_SERVER_LOG(Display, TEXT("FStormSyncTransportServerEndpoint::HandleSyncRequestMessage - Received sync request Message %s (from %s)"), *InMessage.ToString(), *InMessageContext->GetSender().ToString())
+	UE_LOG(LogStormSyncServer, Display, TEXT("FStormSyncTransportServerEndpoint::HandleSyncRequestMessage - Received sync request Message %s (from %s)"), *InMessage.ToString(), *InMessageContext->GetSender().ToString());
 	SendSyncResponse(InMessage, InMessageContext);
 }
 
@@ -238,7 +238,7 @@ void FStormSyncTransportServerEndpoint::SendSyncResponse(const FStormSyncTranspo
 		TUniquePtr<FStormSyncTransportPushResponse> Message(FMessageEndpoint::MakeMessage<FStormSyncTransportPushResponse>(MessageId));
 		if (!Message.IsValid())
 		{
-			STORM_SYNC_SERVER_LOG(Error, TEXT("FStormSyncTransportServerEndpoint::SendSyncResponse - Sync response message is invalid"))
+			UE_LOG(LogStormSyncServer, Error, TEXT("FStormSyncTransportServerEndpoint::SendSyncResponse - Sync response message is invalid"));
 			return;
 		}
 
@@ -253,10 +253,10 @@ void FStormSyncTransportServerEndpoint::SendSyncResponse(const FStormSyncTranspo
 		Message->HostAddress = FStormSyncTransportNetworkUtils::GetTcpEndpointAddress();
 		Message->HostAdapterAddresses = FStormSyncTransportNetworkUtils::GetLocalAdapterAddresses();
 
-		STORM_SYNC_SERVER_LOG(Display, TEXT("FStormSyncTransportServerEndpoint::SendSyncResponse - Sending back to %s sync response: %s"), *RemoteAddress.ToString(), *Message->ToString())
+		UE_LOG(LogStormSyncServer, Display, TEXT("FStormSyncTransportServerEndpoint::SendSyncResponse - Sending back to %s sync response: %s"), *RemoteAddress.ToString(), *Message->ToString());
 		if (FileModifiers.IsEmpty())
 		{
-			STORM_SYNC_SERVER_LOG(Display, TEXT("\tSending back empty list of modifiers, meaning pak are synced."))
+			UE_LOG(LogStormSyncServer, Display, TEXT("\tSending back empty list of modifiers, meaning pak are synced."));
 		}
 
 		MessageEndpointPtr->Send(Message.Release(), RemoteAddress);
diff --git a/Engine/Plugins/Experimental/StormSync/Source/StormSyncTransportServer/Private/StormSyncTransportServerLog.h b/Engine/Plugins/Experimental/StormSync/Source/StormSyncTransportServer/Private/StormSyncTransportServerLog.h
--- a/Engine/Plugins/Experimental/StormSync/Source/StormSyncTransportServer/Private/StormSyncTransportServerLog.h
+++ b/Engine/Plugins/Experimental/StormSync/Source/StormSyncTransportServer/Private/StormSyncTransportServerLog.h
@@ -5,8 +5,3 @@
 #include "Logging/LogMacros.h"
 
 DECLARE_LOG_CATEGORY_EXTERN(LogStormSyncServer, Display, All);
-
-#define STORM_SYNC_SERVER_LOG(Verbosity, Format, ...) \
-{ \
-    UE_LOG(LogStormSyncServer, Verbosity, Format, ##__VA_ARGS__); \
-}
diff --git a/Engine/Plugins/Experimental/StormSync/Source/StormSyncTransportServer/Private/StormSyncTransportServerModule.cpp b/Engine/Plugins/Experimental/StormSync/Source/StormSyncTransportServer/Private/StormSyncTransportServerModule.cpp
--- a/Engine/Plugins/Experimental/StormSync/Source/StormSyncTransportServer/Private/StormSyncTransportServerModule.cpp
+++ b/Engine/Plugins/Experimental/StormSync/Source/StormSyncTransportServer/Private/StormSyncTransportServerModule.cpp
@@ -153,7 +153,7 @@ void FStormSyncTransportServerModule::RegisterConsoleCommands()
 		FConsoleCommandWithArgsDelegate::CreateLambda([this](const TArray<FString>& Args)
 		{
 			const FString AddressId = GetServerEndpointMessageAddressId();
-			STORM_SYNC_SERVER_LOG(Display, TEXT("StormSync.Server.Debug - EndpointId: %s"), *AddressId)
+			UE_LOG(LogStormSyncServer, Display, TEXT("StormSync.Server.Debug - EndpointId: %s"), *AddressId);
 		}),
 		ECVF_Default
 	));
@@ -171,7 +171,7 @@ void FStormSyncTransportServerModule::UnregisterConsoleCommands()
 
 void FStormSyncTransportServerModule::OnEngineLoopInitComplete()
 {
-	STORM_SYNC_SERVER_LOG(Verbose, TEXT("FStormSyncTransportServerModule::OnEngineLoopInitComplete - Publish ping messages for discover manager and server endpoint ..."))
+	UE_LOG(LogStormSyncServer, Verbose, TEXT("FStormSyncTransportServerModule::OnEngineLoopInitComplete - Publish ping messages for discover manager and server endpoint ..."));
 	
 	// We broadcast a message to notify others about this editor instance (this is required so that further "direct" send are received on the other end)
 
@@ -194,7 +194,7 @@ void FStormSyncTransportServerModule::PublishConnectMessage() const
 {
 	if (!DiscoveryManager.IsValid())
 	{
-		STORM_SYNC_SERVER_LOG(Error, TEXT("FStormSyncTransportServerModule::PublishConnectMessage - Unable to send connect message cause DiscoveryManager is invalid"))
+		UE_LOG(LogStormSyncServer, Error, TEXT("FStormSyncTransportServerModule::PublishConnectMessage - Unable to send connect message cause DiscoveryManager is invalid"));
 		return;
 	}
 
@@ -211,7 +211,7 @@ void FStormSyncTransportServerModule::ExecuteStartServer(const TArray<FString>&
 	check(ServerEndpoint.IsValid());
 	if (IsRunning())
 	{
-		STORM_SYNC_SERVER_LOG(Warning, TEXT("FStormSyncTransportServerModule::ExecuteStartServer - Server endpoint TCP listener already running"))
+		UE_LOG(LogStormSyncServer, Warning, TEXT("FStormSyncTransportServerModule::ExecuteStartServer - Server endpoint TCP listener already running"));
 		return;
 	}
 	
@@ -230,7 +230,7 @@ void FStormSyncTransportServerModule::ExecuteStopServer(const TArray<FString>& A
 	}
 	else
 	{
-		STORM_SYNC_SERVER_LOG(Warning, TEXT("FStormSyncTransportServerModule::ExecuteStopServer - Server endpoint already inactive"))
+		UE_LOG(LogStormSyncServer, Warning, TEXT("FStormSyncTransportServerModule::ExecuteStopServer - Server endpoint already inactive"));
 	}
 }
 
@@ -239,7 +239,7 @@ void FStormSyncTransportServerModule::ExecuteServerStatus(const TArray<FString>&
 	FText StatusText;
 	GetServerStatus(StatusText);
 
-	STORM_SYNC_SERVER_LOG(Display, TEXT("FStormSyncTransportServerModule::ExecuteServerStatus - %s"), *StatusText.ToString())
+	UE_LOG(LogStormSyncServer, Display, TEXT("FStormSyncTransportServerModule::ExecuteServerStatus - %s"), *StatusText.ToString());
 }
 
 IMPLEMENT_MODULE(FStormSyncTransportServerModule, StormSyncTransportServer)
diff --git a/Engine/Plugins/Experimental/StormSync/Source/StormSyncTransportServer/Private/Tests/StormSyncTransportServerEndpoint.spec.cpp b/Engine/Plugins/Experimental/StormSync/Source/StormSyncTransportServer/Private/Tests/StormSyncTransportServerEndpoint.spec.cpp
--- a/Engine/Plugins/Experimental/StormSync/Source/StormSyncTransportServer/Private/Tests/StormSyncTransportServerEndpoint.spec.cpp
+++ b/Engine/Plugins/Experimental/StormSync/Source/StormSyncTransportServer/Private/Tests/StormSyncTransportServerEndpoint.spec.cpp
@@ -12,14 +12,6 @@
 #include "StormSyncTransportServerEndpoint.h"
 #include "StormSyncTransportSettings.h"
 
-#define DEFINE_TEST_MESSAGE_HANDLER(DelegateType, MessageType) \
-	DECLARE_DELEGATE_OneParam(FOn##DelegateType, const MessageType&) \
-	FOn##DelegateType On##DelegateType; \
-	void Handle##DelegateType(const MessageType& InMessage, const TSharedRef<IMessageContext, ESPMode::ThreadSafe>& MessageContext) \
-	{ \
-		On##DelegateType.ExecuteIfBound(InMessage); \
-	}
-
 BEGIN_DEFINE_SPEC(FStormSyncTransportServerEndpointSpec, "StormSync.StormSyncTransportServer.StormSyncTransportServerEndpoint", EAutomationTestFlags::ProductFilter | EAutomationTestFlags::ApplicationContextMask)
 
 	static constexpr const TCHAR* EndpointFriendlyName = TEXT("Test Suite ServerEndpoint (Message System)");
@@ -36,8 +28,19 @@ BEGIN_DEFINE_SPEC(FStormSyncTransportServerEndpointSpec, "StormSync.StormSyncTra
 	/** Default to roughly 4 Mb. Socket Send / Receive Buffer Size */
 	const int32 SocketBufferSize = 4 * 1024 * 1024;
 
-	DEFINE_TEST_MESSAGE_HANDLER(PongMessage, FStormSyncTransportPongMessage);
-	DEFINE_TEST_MESSAGE_HANDLER(PushResponse, FStormSyncTransportPushResponse);
+	TDelegate<void(const FStormSyncTransportPongMessage&)> OnPongMessage;
+
+	TDelegate<void(const FStormSyncTransportPushResponse&)> OnPushResponse;
+
+	void HandlePongMessage(const FStormSyncTransportPongMessage& InMessage, const TSharedRef<IMessageContext, ESPMode::ThreadSafe>&)
+	{
+		OnPongMessage.ExecuteIfBound(InMessage);
+	}
+
+	void HandlePushResponse(const FStormSyncTransportPushResponse& InMessage, const TSharedRef<IMessageContext, ESPMode::ThreadSafe>&)
+	{
+		OnPushResponse.ExecuteIfBound(InMessage);
+	}
 
 	static uint64 GetMockFileSize(const FName& InPackageName)
 	{
