Storm Sync: disallow autostart server in commandlets.
Auto-start server was reporting an error, causing the packaging to fail

diff --git a/Engine/Plugins/Experimental/StormSync/Source/StormSyncTransportServer/Private/StormSyncTransportServerModule.cpp b/Engine/Plugins/Experimental/StormSync/Source/StormSyncTransportServer/Private/StormSyncTransportServerModule.cpp
--- a/Engine/Plugins/Experimental/StormSync/Source/StormSyncTransportServer/Private/StormSyncTransportServerModule.cpp
+++ b/Engine/Plugins/Experimental/StormSync/Source/StormSyncTransportServer/Private/StormSyncTransportServerModule.cpp
@@ -29,7 +29,10 @@ void FStormSyncTransportServerModule::StartupModule()
 
 	// Create server endpoint and start tcp listening if configured to do so
 	ServerEndpoint = CreateServerLocalEndpoint(TEXT("Server"));
-	const bool bIsCommandLineAutoStartDisabled = UE::StormSync::Transport::Private::IsServerAutoStartDisabled();
+
+	// Auto-start is disabled if running commandlet
+	const bool bIsCommandLineAutoStartDisabled = IsRunningCommandlet() || UE::StormSync::Transport::Private::IsServerAutoStartDisabled();
+
 	if (!bIsCommandLineAutoStartDisabled && Settings->IsAutoStartServer())
 	{
 		if (ServerEndpoint->StartTcpListener() && IsRunning())
