[StormSync] StormSync Transport modules: automatic start of endpoints is disabled by default in commandlets.

Code has been reworked so that it is possible to start the endpoints and run services using eitehr console commands or the module's API, even within a commandlet. So, it would be possible to execute a StormSync commandlets to be used in a back end for instance.

- The console commands are registered even in commandlet so they can be used.
- Added Client.Start console command.
- Care is taken to publish the ping messages after the engine loop complete and is done when the service is started.
- Using OnPostEngineInit instead of OnFEngineLoopInitComplete because the later is not called when running commandlets while the former is.

Tests:
- Tested cooking project and validated that the transport endpoints are no longer created.
- Tested StormSync commands from the browser menu between 2 editor instances.
- Implemented a commandlet with StormSync commands to validate that it can connect to other instances and execute commands with the manual start path.

diff --git a/Engine/Plugins/Experimental/StormSync/Source/StormSyncDrives/Private/StormSyncDrivesModule.cpp b/Engine/Plugins/Experimental/StormSync/Source/StormSyncDrives/Private/StormSyncDrivesModule.cpp
--- a/Engine/Plugins/Experimental/StormSync/Source/StormSyncDrives/Private/StormSyncDrivesModule.cpp
+++ b/Engine/Plugins/Experimental/StormSync/Source/StormSyncDrives/Private/StormSyncDrivesModule.cpp
@@ -27,7 +27,7 @@ void FStormSyncDrivesModule::StartupModule()
 	LogListing = MessageLogModule.GetLogListing(LogName);
 #endif
 
-	FCoreDelegates::OnFEngineLoopInitComplete.AddRaw(this, &FStormSyncDrivesModule::OnEngineLoopInitComplete);
+	FCoreDelegates::OnPostEngineInit.AddRaw(this, &FStormSyncDrivesModule::OnPostEngineInit);
 
 #if WITH_EDITOR
 	GetMutableDefault<UStormSyncDrivesSettings>()->OnSettingChanged().AddRaw(this, &FStormSyncDrivesModule::OnSettingsChanged);
@@ -45,7 +45,7 @@ void FStormSyncDrivesModule::ShutdownModule()
 	}
 #endif
 
-	FCoreDelegates::OnFEngineLoopInitComplete.RemoveAll(this);
+	FCoreDelegates::OnPostEngineInit.RemoveAll(this);
 
 #if WITH_EDITOR
 	if (UObjectInitialized())
@@ -118,9 +118,9 @@ bool FStormSyncDrivesModule::UnregisterMountPoint(const FStormSyncMountPointConf
 	return true;
 }
 
-void FStormSyncDrivesModule::OnEngineLoopInitComplete()
+void FStormSyncDrivesModule::OnPostEngineInit()
 {
-	UE_LOG(LogStormSyncDrives, Display, TEXT("FStormSyncDrivesModule::OnEngineLoopInitComplete ..."));
+	UE_LOG(LogStormSyncDrives, Display, TEXT("FStormSyncDrivesModule::OnPostEngineInit ..."));
 	UE_LOG(LogStormSyncDrives, Display, TEXT("\t Mounting Drives based on config"));
 
 	ResetMountedDrivesFromSettings(GetDefault<UStormSyncDrivesSettings>());
diff --git a/Engine/Plugins/Experimental/StormSync/Source/StormSyncDrives/Private/StormSyncDrivesModule.h b/Engine/Plugins/Experimental/StormSync/Source/StormSyncDrives/Private/StormSyncDrivesModule.h
--- a/Engine/Plugins/Experimental/StormSync/Source/StormSyncDrives/Private/StormSyncDrivesModule.h
+++ b/Engine/Plugins/Experimental/StormSync/Source/StormSyncDrives/Private/StormSyncDrivesModule.h
@@ -42,8 +42,8 @@ private:
 	TSharedPtr<IMessageLogListing> LogListing;
 #endif
 
-	/** Core OnEngineLoopInitComplete delegate used to mount any stored config in developer settings */
-	void OnEngineLoopInitComplete();
+	/** Core OnPostEngineInit delegate used to mount any stored config in developer settings */
+	void OnPostEngineInit();
 
 	/**
 	 * Developer settings changed handler. Called anytime settings changed.
diff --git a/Engine/Plugins/Experimental/StormSync/Source/StormSyncTransportClient/Private/StormSyncTransportClientModule.cpp b/Engine/Plugins/Experimental/StormSync/Source/StormSyncTransportClient/Private/StormSyncTransportClientModule.cpp
--- a/Engine/Plugins/Experimental/StormSync/Source/StormSyncTransportClient/Private/StormSyncTransportClientModule.cpp
+++ b/Engine/Plugins/Experimental/StormSync/Source/StormSyncTransportClient/Private/StormSyncTransportClientModule.cpp
@@ -15,20 +15,21 @@
 
 void FStormSyncTransportClientModule::StartupModule()
 {
-	ClientEndpoint = CreateClientLocalEndpoint(TEXT("Client"));
-
-	if (GIsEditor && !IsRunningCommandlet())
+	// Auto-start the client unless in a commandlet.
+	if (!IsRunningCommandlet())
 	{
-		RegisterConsoleCommands();
+		StartClientEndpoint(DefaultClientEndpointName);
 	}
+
+	RegisterConsoleCommands();
 	
 	// Register for engine initialization completed so we can broadcast presence over the network from this client to other clients
-	FCoreDelegates::OnFEngineLoopInitComplete.AddRaw(this, &FStormSyncTransportClientModule::OnEngineLoopInitComplete);
+	FCoreDelegates::OnPostEngineInit.AddRaw(this, &FStormSyncTransportClientModule::OnPostEngineInit);
 }
 
 void FStormSyncTransportClientModule::ShutdownModule()
 {
-	FCoreDelegates::OnFEngineLoopInitComplete.RemoveAll(this);
+	FCoreDelegates::OnPostEngineInit.RemoveAll(this);
 	
 	if (ClientEndpoint.IsValid())
 	{
@@ -38,11 +39,29 @@ void FStormSyncTransportClientModule::ShutdownModule()
 	UnregisterConsoleCommands();
 }
 
+void FStormSyncTransportClientModule::StartClientEndpoint(const FString& InEndpointFriendlyName)
+{
+	if (!ClientEndpoint.IsValid())
+	{
+		ClientEndpoint = CreateClientLocalEndpoint(InEndpointFriendlyName);
+		if (!ClientEndpoint.IsValid())
+		{
+			UE_LOG(LogStormSyncClient, Error, TEXT("FStormSyncTransportClientModule::StartClientEndpoint - Failed to create Client Local Endpoint"));
+			return;
+		}
+	}
+	
+	if (bEngineInitComplete)
+	{
+		PublishStatusPingMessage();
+	}
+}
+
 TSharedPtr<IStormSyncTransportClientLocalEndpoint> FStormSyncTransportClientModule::CreateClientLocalEndpoint(const FString& InEndpointFriendlyName) const
 {
 	TSharedPtr<FStormSyncTransportClientEndpoint, ESPMode::ThreadSafe> Endpoint = MakeShared<FStormSyncTransportClientEndpoint>();
 	check(Endpoint.IsValid());
-	Endpoint->InitializeMessaging(TEXT("Client"));
+	Endpoint->InitializeMessaging(InEndpointFriendlyName);
 	return Endpoint;
 }
 
@@ -112,7 +131,6 @@ void FStormSyncTransportClientModule::PushPackages(const FStormSyncPackageDescri
 {
 	UE_LOG(LogStormSyncClient, Display, TEXT("FStormSyncTransportClientModule::PushPackages - PackageDescriptor: %s, InPackageNames: %d, MessageAddressId: %s"), *InPackageDescriptor.ToString(), InPackageNames.Num(), *InMessageAddress.ToString());
 
-	const FMessageEndpointSharedPtr MessageEndpoint = GetClientMessageEndpoint();
 	if (!ClientEndpoint.IsValid())
 	{
 		UE_LOG(LogStormSyncClient, Error, TEXT("FStormSyncTransportClientModule::PushPackages - Unable to get client endpoint"));
@@ -126,8 +144,6 @@ void FStormSyncTransportClientModule::PullPackages(const FStormSyncPackageDescri
 {
 	UE_LOG(LogStormSyncClient, Display, TEXT("FStormSyncTransportClientModule::PullPackages - PackageDescriptor: %s, InPackageNames: %d, MessageAddressId: %s"), *InPackageDescriptor.ToString(), InPackageNames.Num(), *InMessageAddress.ToString());
 
-	const FMessageEndpointSharedPtr MessageEndpoint = GetClientMessageEndpoint();
-
 	if (!ClientEndpoint.IsValid())
 	{
 		UE_LOG(LogStormSyncClient, Error, TEXT("FStormSyncTransportClientModule::PullPackages - Unable to get client endpoint"));
@@ -141,7 +157,6 @@ void FStormSyncTransportClientModule::RequestPackagesStatus(const FMessageAddres
 {
 	UE_LOG(LogStormSyncClient, Display, TEXT("FStormSyncTransportClientModule::RequestPackagesStatus - InRemoteAddress: %s, InPackageNames: %d"), *InRemoteAddress.ToString(), InPackageNames.Num());
 
-	const FMessageEndpointSharedPtr MessageEndpoint = GetClientMessageEndpoint();
 	if (!ClientEndpoint.IsValid())
 	{
 		UE_LOG(LogStormSyncClient, Error, TEXT("FStormSyncTransportClientModule::RequestPackagesStatus - Unable to get client endpoint"));
@@ -153,6 +168,13 @@ void FStormSyncTransportClientModule::RequestPackagesStatus(const FMessageAddres
 
 void FStormSyncTransportClientModule::RegisterConsoleCommands()
 {
+	ConsoleCommands.Add(IConsoleManager::Get().RegisterConsoleCommand(
+		TEXT("StormSync.Client.Start"),
+		TEXT("Starts Storm Sync Client. Usage: [EndpointName]"),
+		FConsoleCommandWithArgsDelegate::CreateRaw(this, &FStormSyncTransportClientModule::ExecuteStartClient),
+		ECVF_Default
+	));
+
 	ConsoleCommands.Add(IConsoleManager::Get().RegisterConsoleCommand(
 		TEXT("StormSync.Client.Ping"),
 		TEXT("Sends a ping message on Storm Sync message bus"),
@@ -170,46 +192,14 @@ void FStormSyncTransportClientModule::RegisterConsoleCommands()
 	ConsoleCommands.Add(IConsoleManager::Get().RegisterConsoleCommand(
 		TEXT("StormSync.Client.Debug"),
 		TEXT("Prints out client address enpoint id to the console"),
-		// FConsoleCommandWithArgsDelegate::CreateRaw(this, &FStormSyncTransportClientModule::ExecuteSyncPak),
-		FConsoleCommandWithArgsDelegate::CreateLambda([this](const TArray<FString>& Args)
-		{
-			const FString AddressId = GetClientEndpointMessageAddressId();
-			UE_LOG(LogStormSyncClient, Display, TEXT("StormSync.Client.Debug - EndpointId: %s"), *AddressId);
-		}),
+		FConsoleCommandWithArgsDelegate::CreateRaw(this, &FStormSyncTransportClientModule::ExecuteDebug),
 		ECVF_Default
 	));
 
 	ConsoleCommands.Add(IConsoleManager::Get().RegisterConsoleCommand(
 		TEXT("StormSync.Client.Debug.Ping"),
 		TEXT("Sends status ping"),
-		// FConsoleCommandWithArgsDelegate::CreateRaw(this, &FStormSyncTransportClientModule::ExecuteSyncPak),
-		FConsoleCommandWithArgsDelegate::CreateLambda([this](const TArray<FString>& Args)
-		{
-			if (ClientEndpoint.IsValid())
-			{
-				const TSharedPtr<FMessageEndpoint, ESPMode::ThreadSafe> MessageEndpoint = ClientEndpoint->GetMessageEndpoint();
-				if (!MessageEndpoint.IsValid())
-				{
-					UE_LOG(LogStormSyncClient, Error, TEXT("StormSync.Client.Debug.Ping - Unable to send Connect Message cause Message Endpoint is invalid"));
-					return;
-				}
-
-				UE_LOG(LogStormSyncClient, Display, TEXT("StormSync.Client.Debug.Ping - Publish Client Connect Message ..."));
-				FStormSyncTransportStatusPing* Message = FMessageEndpoint::MakeMessage<FStormSyncTransportStatusPing>();
-				MessageEndpoint->Publish(Message);
-			}
-		}),
-		ECVF_Default
-	));
-	
-	ConsoleCommands.Add(IConsoleManager::Get().RegisterConsoleCommand(
-		TEXT("StormSync.Client.InstanceId"),
-		TEXT("Prints out local app instance id"),
-		FConsoleCommandWithArgsDelegate::CreateLambda([this](const TArray<FString>& Args)
-		{
-			const FString InstanceId = FApp::GetInstanceId().ToString();
-			UE_LOG(LogStormSyncClient, Display, TEXT("StormSync.Client.Debug - InstanceId: %s"), *InstanceId);
-		}),
+		FConsoleCommandWithArgsDelegate::CreateRaw(this, &FStormSyncTransportClientModule::ExecuteDebugPing),
 		ECVF_Default
 	));
 }
@@ -224,27 +214,39 @@ void FStormSyncTransportClientModule::UnregisterConsoleCommands()
 	ConsoleCommands.Empty();
 }
 
-void FStormSyncTransportClientModule::OnEngineLoopInitComplete()
+void FStormSyncTransportClientModule::OnPostEngineInit()
 {
+	bEngineInitComplete = true;
+	
 	if (ClientEndpoint.IsValid())
 	{
-		const TSharedPtr<FMessageEndpoint, ESPMode::ThreadSafe> MessageEndpoint = ClientEndpoint->GetMessageEndpoint();
-		if (!MessageEndpoint.IsValid())
-		{
-			UE_LOG(LogStormSyncClient, Error, TEXT("FStormSyncTransportClientModule::OnEngineLoopInitComplete - Unable to send Connect Message cause Message Endpoint is invalid"));
-			return;
-		}
+		PublishStatusPingMessage();
+	}
+}
 
-		// We broadcast a message to notify others about this editor instance (this is required so that further "direct" send are received on the other end)
-		UE_LOG(LogStormSyncClient, Display, TEXT("FStormSyncTransportClientModule::OnEngineLoopInitComplete - Publish Client Connect Message ..."));
-		FStormSyncTransportStatusPing* Message = FMessageEndpoint::MakeMessage<FStormSyncTransportStatusPing>();
-		MessageEndpoint->Publish(Message);
+void FStormSyncTransportClientModule::PublishStatusPingMessage() const
+{
+	const FMessageEndpointSharedPtr MessageEndpoint = GetClientMessageEndpoint();
+	if (!MessageEndpoint.IsValid())
+	{
+		UE_LOG(LogStormSyncClient, Error, TEXT("FStormSyncTransportClientModule::PublishStatusPingMessage - Unable to send Connect Message cause Message Endpoint is invalid"));
+		return;
 	}
+
+	// We broadcast a message to notify others about this editor instance (this is required so that further "direct" send are received on the other end)
+	UE_LOG(LogStormSyncClient, Display, TEXT("FStormSyncTransportClientModule::PublishStatusPingMessage - Publish Client Connect Message ..."));
+	FStormSyncTransportStatusPing* Message = FMessageEndpoint::MakeMessage<FStormSyncTransportStatusPing>();
+	MessageEndpoint->Publish(Message);
+}
+
+void FStormSyncTransportClientModule::ExecuteStartClient(const TArray<FString>& InArgs)
+{
+	StartClientEndpoint(InArgs.Num() ? InArgs[0] : DefaultClientEndpointName);	
 }
 
-void FStormSyncTransportClientModule::ExecutePing(const TArray<FString>& Args)
+void FStormSyncTransportClientModule::ExecutePing(const TArray<FString>& InArgs)
 {
-	const FString Argv = FString::Join(Args, TEXT(""));
+	const FString Argv = FString::Join(InArgs, TEXT(""));
 	UE_LOG(LogStormSyncClient, Display, TEXT("FStormSyncTransportClientModule::ExecutePing - %s"), *Argv);
 
 	const FMessageEndpointSharedPtr MessageEndpoint = GetClientMessageEndpoint();
@@ -257,25 +259,25 @@ void FStormSyncTransportClientModule::ExecutePing(const TArray<FString>& Args)
 	MessageEndpoint->Publish(FMessageEndpoint::MakeMessage<FStormSyncTransportPingMessage>());
 }
 
-FStormSyncPackageDescriptor FStormSyncTransportClientModule::CreatePackageDescriptorFromCommandLine(const FString& Argv)
+FStormSyncPackageDescriptor FStormSyncTransportClientModule::CreatePackageDescriptorFromCommandLine(const FString& InArgv)
 {
 	FStormSyncPackageDescriptor PackageDescriptor;
-	if (!FParse::Value(*Argv, TEXT("-name="), PackageDescriptor.Name))
+	if (!FParse::Value(*InArgv, TEXT("-name="), PackageDescriptor.Name))
 	{
 		PackageDescriptor.Name = DefaultPakName;
 		UE_LOG(LogStormSyncClient, Display, TEXT("FStormSyncEditorModule::CreatePackageDescriptorFromCommandLine - Missing -name parameter, using default \"%s\""), *PackageDescriptor.Name);
 	}
 
-	FParse::Value(*Argv, TEXT("-version="), PackageDescriptor.Version);
-	FParse::Value(*Argv, TEXT("-description="), PackageDescriptor.Description);
-	FParse::Value(*Argv, TEXT("-author="), PackageDescriptor.Author);
+	FParse::Value(*InArgv, TEXT("-version="), PackageDescriptor.Version);
+	FParse::Value(*InArgv, TEXT("-description="), PackageDescriptor.Description);
+	FParse::Value(*InArgv, TEXT("-author="), PackageDescriptor.Author);
 
 	return PackageDescriptor;
 }
 
-void FStormSyncTransportClientModule::ExecuteSyncPak(const TArray<FString>& Args)
+void FStormSyncTransportClientModule::ExecuteSyncPak(const TArray<FString>& InArgs)
 {
-	const FString Argv = FString::Join(Args, TEXT(" "));
+	const FString Argv = FString::Join(InArgs, TEXT(" "));
 	UE_LOG(LogStormSyncClient, Display, TEXT("FStormSyncTransportClientModule::ExecuteSyncPak - Argv: %s"), *Argv);
 
 	// Parse command line.
@@ -295,6 +297,28 @@ void FStormSyncTransportClientModule::ExecuteSyncPak(const TArray<FString>& Args
 	SynchronizePackages(PackageDescriptor, PackageNames);
 }
 
+void FStormSyncTransportClientModule::ExecuteDebug(const TArray<FString>& InArgs)
+{
+	const FString AddressId = GetClientEndpointMessageAddressId();
+	UE_LOG(LogStormSyncClient, Display, TEXT("StormSync.Client.Debug - EndpointId: %s"), *AddressId);
+	const FString InstanceId = FApp::GetInstanceId().ToString();
+	UE_LOG(LogStormSyncClient, Display, TEXT("StormSync.Client.Debug - InstanceId: %s"), *InstanceId);
+}
+
+void FStormSyncTransportClientModule::ExecuteDebugPing(const TArray<FString>& InArgs)
+{
+	const FMessageEndpointSharedPtr MessageEndpoint = GetClientMessageEndpoint();
+	if (!MessageEndpoint.IsValid())
+	{
+		UE_LOG(LogStormSyncClient, Error, TEXT("StormSync.Client.Debug.Ping - Unable to send Connect Message cause Message Endpoint is invalid"));
+		return;
+	}
+
+	UE_LOG(LogStormSyncClient, Display, TEXT("StormSync.Client.Debug.Ping - Publish Client Connect Message ..."));
+	FStormSyncTransportStatusPing* Message = FMessageEndpoint::MakeMessage<FStormSyncTransportStatusPing>();
+	MessageEndpoint->Publish(Message);
+}
+
 IMPLEMENT_MODULE(FStormSyncTransportClientModule, StormSyncTransportClient)
 
 #undef LOCTEXT_NAMESPACE
diff --git a/Engine/Plugins/Experimental/StormSync/Source/StormSyncTransportClient/Private/StormSyncTransportClientModule.h b/Engine/Plugins/Experimental/StormSync/Source/StormSyncTransportClient/Private/StormSyncTransportClientModule.h
--- a/Engine/Plugins/Experimental/StormSync/Source/StormSyncTransportClient/Private/StormSyncTransportClientModule.h
+++ b/Engine/Plugins/Experimental/StormSync/Source/StormSyncTransportClient/Private/StormSyncTransportClientModule.h
@@ -19,6 +19,7 @@ public:
 	virtual void ShutdownModule() override;
 
 	//~ Begin IStormSyncTransportClientModule
+	virtual void StartClientEndpoint(const FString& InEndpointFriendlyName) override;
 	virtual TSharedPtr<IStormSyncTransportClientLocalEndpoint> CreateClientLocalEndpoint(const FString& InEndpointFriendlyName) const override;
 	virtual FString GetClientEndpointMessageAddressId() const override;
 	virtual FMessageEndpointSharedPtr GetClientMessageEndpoint() const override;
@@ -29,15 +30,6 @@ public:
 	//~ End IStormSyncTransportClientModule
 
 private:
-	/** The name of the ava pak to use when none is provided from command line arguments */
-	static constexpr const TCHAR* DefaultPakName = TEXT("SyncPak");
-	
-	/** Our message endpoint provider */
-	TSharedPtr<IStormSyncTransportClientLocalEndpoint, ESPMode::ThreadSafe> ClientEndpoint;
-	
-	/** References of registered console commands via IConsoleManager */
-	TArray<IConsoleObject*> ConsoleCommands;
-
 	/** Called from StartupModule and sets up console commands for the plugin via IConsoleManager */
 	void RegisterConsoleCommands();
 
@@ -45,14 +37,42 @@ private:
 	void UnregisterConsoleCommands();
 	
 	/** Event handler to kick in operations once engine is fully initialized (to publish a client connect message) */
-	void OnEngineLoopInitComplete();
+	void OnPostEngineInit();
 
+	/** Publish the client endpoint's ping message. */
+	void PublishStatusPingMessage() const;
+
+	/** Command handler for starting the client. */
+	void ExecuteStartClient(const TArray<FString>& InArgs);
+	
 	/** Command handler for ping command */
-	void ExecutePing(const TArray<FString>& Args);
+	void ExecutePing(const TArray<FString>& InArgs);
 	
 	/** Command handler for sync pak command */
-	void ExecuteSyncPak(const TArray<FString>& Args);
-	
+	void ExecuteSyncPak(const TArray<FString>& InArgs);
+
+	/** Command handler for displaying debug info. */
+	void ExecuteDebug(const TArray<FString>& InArgs);
+
+	/** Command handler for sending a status ping. */
+	void ExecuteDebugPing(const TArray<FString>& InArgs);
+
 	/** Returns a new package descriptor pulling info from command line options */
-	static FStormSyncPackageDescriptor CreatePackageDescriptorFromCommandLine(const FString& Argv);
+	static FStormSyncPackageDescriptor CreatePackageDescriptorFromCommandLine(const FString& InArgv);
+
+private:
+	/** The name of the ava pak to use when none is provided from command line arguments */
+	static constexpr const TCHAR* DefaultPakName = TEXT("SyncPak");
+
+	/** Default client endpoint name. */
+	static constexpr const TCHAR* DefaultClientEndpointName = TEXT("Client");
+	
+	/** Our message endpoint provider */
+	TSharedPtr<IStormSyncTransportClientLocalEndpoint, ESPMode::ThreadSafe> ClientEndpoint;
+	
+	/** References of registered console commands via IConsoleManager */
+	TArray<IConsoleObject*> ConsoleCommands;
+
+	/** Indicates when the engine init is complete. This is used to know when pending messages can be published. */
+	bool bEngineInitComplete = false;
 };
diff --git a/Engine/Plugins/Experimental/StormSync/Source/StormSyncTransportClient/Public/IStormSyncTransportClientModule.h b/Engine/Plugins/Experimental/StormSync/Source/StormSyncTransportClient/Public/IStormSyncTransportClientModule.h
--- a/Engine/Plugins/Experimental/StormSync/Source/StormSyncTransportClient/Public/IStormSyncTransportClientModule.h
+++ b/Engine/Plugins/Experimental/StormSync/Source/StormSyncTransportClient/Public/IStormSyncTransportClientModule.h
@@ -37,8 +37,11 @@ public:
 		static const FName ModuleName = "StormSyncTransportClient";
 		return FModuleManager::Get().IsModuleLoaded(ModuleName);
 	}
+
+	/** Start transport client endpoint */
+	virtual void StartClientEndpoint(const FString& InEndpointFriendlyName) = 0;
 	
-	/** Create a local transport endpoint */
+	UE_DEPRECATED(5.5, "Use StartClientEndpoint() instead.")
 	virtual TSharedPtr<IStormSyncTransportClientLocalEndpoint> CreateClientLocalEndpoint(const FString& InEndpointFriendlyName) const = 0;
 	
 	/** Returns Message Address UID for client endpoint if it is currently running, empty string otherwise */
diff --git a/Engine/Plugins/Experimental/StormSync/Source/StormSyncTransportServer/Private/StormSyncTransportServerModule.cpp b/Engine/Plugins/Experimental/StormSync/Source/StormSyncTransportServer/Private/StormSyncTransportServerModule.cpp
--- a/Engine/Plugins/Experimental/StormSync/Source/StormSyncTransportServer/Private/StormSyncTransportServerModule.cpp
+++ b/Engine/Plugins/Experimental/StormSync/Source/StormSyncTransportServer/Private/StormSyncTransportServerModule.cpp
@@ -19,40 +19,23 @@ void FStormSyncTransportServerModule::StartupModule()
 {
 	const UStormSyncTransportSettings* Settings = GetDefault<UStormSyncTransportSettings>();
 
-	HeartbeatEmitter = MakeUnique<FStormSyncHeartbeatEmitter>();
-	DiscoveryManager = MakeUnique<FStormSyncDiscoveryManager>(
-		Settings->GetMessageBusHeartbeatTimeout(),
-		Settings->GetMessageBusTimeBeforeRemovingInactiveSource(),
-		Settings->GetDiscoveryManagerTickInterval(),
-		Settings->IsDiscoveryPeriodicPublishEnabled()
-	);
-
-	// Create server endpoint and start tcp listening if configured to do so
-	ServerEndpoint = CreateServerLocalEndpoint(TEXT("Server"));
-
 	// Auto-start is disabled if running commandlet
 	const bool bIsCommandLineAutoStartDisabled = IsRunningCommandlet() || UE::StormSync::Transport::Private::IsServerAutoStartDisabled();
 
 	if (!bIsCommandLineAutoStartDisabled && Settings->IsAutoStartServer())
 	{
-		if (ServerEndpoint->StartTcpListener() && IsRunning())
-		{
-			FStormSyncCoreDelegates::OnStormSyncServerStarted.Broadcast();
-		}
+		ExecuteStartServer({});
 	}
 
-	if (GIsEditor && !IsRunningCommandlet())
-	{
-		RegisterConsoleCommands();
-	}
+	RegisterConsoleCommands();
 
 	// Register for engine initialization completed so we can broadcast presence over the network and start heartbeats
-	FCoreDelegates::OnFEngineLoopInitComplete.AddRaw(this, &FStormSyncTransportServerModule::OnEngineLoopInitComplete);
+	FCoreDelegates::OnPostEngineInit.AddRaw(this, &FStormSyncTransportServerModule::OnPostEngineInit);
 }
 
 void FStormSyncTransportServerModule::ShutdownModule()
 {
-	FCoreDelegates::OnFEngineLoopInitComplete.RemoveAll(this);
+	FCoreDelegates::OnPostEngineInit.RemoveAll(this);
 
 	if (ServerEndpoint.IsValid())
 	{
@@ -73,6 +56,43 @@ void FStormSyncTransportServerModule::ShutdownModule()
 	UnregisterConsoleCommands();
 }
 
+void FStormSyncTransportServerModule::StartDiscoveryManager()
+{
+	ConditionalStartHeartBeatEmitter();
+	ConditionalStartDiscoveryManager();
+}
+
+void FStormSyncTransportServerModule::StartServerEndpoint(const FString& InEndpointFriendlyName)
+{
+	StartDiscoveryManager();
+	
+	if (!ServerEndpoint.IsValid())
+	{
+		ServerEndpoint = CreateServerLocalEndpoint(TEXT("Server"));
+		if (!ServerEndpoint.IsValid())
+		{
+			UE_LOG(LogStormSyncServer, Error, TEXT("FStormSyncTransportServerModule::StartServerEndpoint - Failed to create Server Local Endpoint"));
+			return;
+		}
+	}
+	
+	if (IsRunning())
+	{
+		UE_LOG(LogStormSyncServer, Warning, TEXT("FStormSyncTransportServerModule::StartServerEndpoint - Server endpoint TCP listener already running"));
+		return;
+	}
+	
+	if (ServerEndpoint->StartTcpListener() && IsRunning())
+	{
+		FStormSyncCoreDelegates::OnStormSyncServerStarted.Broadcast();
+	}
+
+	if (bEngineInitComplete)
+	{
+		PublishPingMessage();
+	}
+}
+
 TSharedPtr<IStormSyncTransportServerLocalEndpoint> FStormSyncTransportServerModule::CreateServerLocalEndpoint(const FString& InEndpointFriendlyName) const
 {
 	TSharedPtr<FStormSyncTransportServerEndpoint, ESPMode::ThreadSafe> Endpoint = MakeShared<FStormSyncTransportServerEndpoint>();
@@ -101,11 +121,11 @@ bool FStormSyncTransportServerModule::IsRunning() const
 	return ServerEndpoint.IsValid() && ServerEndpoint->IsRunning() && ServerEndpoint->IsTcpServerActive();
 }
 
-bool FStormSyncTransportServerModule::GetServerStatus(FText& StatusText) const
+bool FStormSyncTransportServerModule::GetServerStatus(FText& OutStatusText) const
 {
 	if (!ServerEndpoint.IsValid())
 	{
-		StatusText = LOCTEXT("ServerStatusEndpointInvalid", "Server is not active.");
+		OutStatusText = LOCTEXT("ServerStatusEndpointInvalid", "Server is not active.");
 		return false;
 	}
 
@@ -113,19 +133,47 @@ bool FStormSyncTransportServerModule::GetServerStatus(FText& StatusText) const
 
 	if (bIsRunning)
 	{
-		StatusText = FText::Format(
+		OutStatusText = FText::Format(
 			LOCTEXT("ServerStatusEndpointRunning", "Server is currently running and listening for incoming connections on {0}"),
 			FText::FromString(ServerEndpoint->GetTcpServerEndpointAddress())
 		);
 	}
 	else
 	{
-		StatusText = LOCTEXT("ServerStatusEndpointNotRunning", "Server is not running.");
+		OutStatusText = LOCTEXT("ServerStatusEndpointNotRunning", "Server is not running.");
 	}
 	
 	return bIsRunning;
 }
 
+void FStormSyncTransportServerModule::ConditionalStartHeartBeatEmitter()
+{
+	if (!HeartbeatEmitter)
+	{
+		HeartbeatEmitter = MakeUnique<FStormSyncHeartbeatEmitter>();
+	}
+}
+
+void FStormSyncTransportServerModule::ConditionalStartDiscoveryManager()
+{
+	if (!DiscoveryManager)
+	{
+		const UStormSyncTransportSettings* Settings = GetDefault<UStormSyncTransportSettings>();
+
+		DiscoveryManager = MakeUnique<FStormSyncDiscoveryManager>(
+			Settings->GetMessageBusHeartbeatTimeout(),
+			Settings->GetMessageBusTimeBeforeRemovingInactiveSource(),
+			Settings->GetDiscoveryManagerTickInterval(),
+			Settings->IsDiscoveryPeriodicPublishEnabled()
+		);
+
+		if (bEngineInitComplete)
+		{
+			PublishConnectMessage();
+		}
+	}
+}
+
 void FStormSyncTransportServerModule::RegisterConsoleCommands()
 {
 	ConsoleCommands.Add(IConsoleManager::Get().RegisterConsoleCommand(
@@ -172,9 +220,11 @@ void FStormSyncTransportServerModule::UnregisterConsoleCommands()
 	ConsoleCommands.Empty();
 }
 
-void FStormSyncTransportServerModule::OnEngineLoopInitComplete()
+void FStormSyncTransportServerModule::OnPostEngineInit()
 {
-	UE_LOG(LogStormSyncServer, Verbose, TEXT("FStormSyncTransportServerModule::OnEngineLoopInitComplete - Publish ping messages for discover manager and server endpoint ..."));
+	UE_LOG(LogStormSyncServer, Verbose, TEXT("FStormSyncTransportServerModule::OnPostEngineInit - Publish ping messages for discover manager and server endpoint ..."));
+
+	bEngineInitComplete = true;
 	
 	// We broadcast a message to notify others about this editor instance (this is required so that further "direct" send are received on the other end)
 
@@ -182,49 +232,36 @@ void FStormSyncTransportServerModule::OnEngineLoopInitComplete()
 	PublishConnectMessage();
 
 	// For server endpoint
-	if (ServerEndpoint.IsValid())
-	{
-		const TSharedPtr<FMessageEndpoint, ESPMode::ThreadSafe> MessageEndpoint = ServerEndpoint->GetMessageEndpoint();
-		TUniquePtr<FStormSyncTransportPingMessage> Message(FMessageEndpoint::MakeMessage<FStormSyncTransportPingMessage>());
-		if (MessageEndpoint.IsValid() && Message.IsValid())
-		{
-			MessageEndpoint->Publish(Message.Release());
-		}
-	}
+	PublishPingMessage();
 }
 
 void FStormSyncTransportServerModule::PublishConnectMessage() const
 {
-	if (!DiscoveryManager.IsValid())
+	if (DiscoveryManager.IsValid())
 	{
-		UE_LOG(LogStormSyncServer, Error, TEXT("FStormSyncTransportServerModule::PublishConnectMessage - Unable to send connect message cause DiscoveryManager is invalid"));
-		return;
+		DiscoveryManager->PublishConnectMessage();
 	}
-
-	DiscoveryManager->PublishConnectMessage();
 }
 
-void FStormSyncTransportServerModule::ExecuteStartServer(const TArray<FString>& Args)
+void FStormSyncTransportServerModule::PublishPingMessage() const
 {
-	if (!ServerEndpoint.IsValid())
+	if (ServerEndpoint.IsValid())
 	{
-		ServerEndpoint = CreateServerLocalEndpoint(TEXT("Server"));
+		const TSharedPtr<FMessageEndpoint, ESPMode::ThreadSafe> MessageEndpoint = ServerEndpoint->GetMessageEndpoint();
+		TUniquePtr<FStormSyncTransportPingMessage> Message(FMessageEndpoint::MakeMessage<FStormSyncTransportPingMessage>());
+		if (MessageEndpoint.IsValid() && Message.IsValid())
+		{
+			MessageEndpoint->Publish(Message.Release());
+		}
 	}
+}
 
-	check(ServerEndpoint.IsValid());
-	if (IsRunning())
-	{
-		UE_LOG(LogStormSyncServer, Warning, TEXT("FStormSyncTransportServerModule::ExecuteStartServer - Server endpoint TCP listener already running"));
-		return;
-	}
-	
-	if (ServerEndpoint->StartTcpListener() && IsRunning())
-	{
-		FStormSyncCoreDelegates::OnStormSyncServerStarted.Broadcast();
-	}
+void FStormSyncTransportServerModule::ExecuteStartServer(const TArray<FString>& InArgs)
+{
+	StartServerEndpoint(TEXT("Server"));
 }
 
-void FStormSyncTransportServerModule::ExecuteStopServer(const TArray<FString>& Args)
+void FStormSyncTransportServerModule::ExecuteStopServer(const TArray<FString>& InArgs)
 {
 	if (ServerEndpoint.IsValid())
 	{
@@ -237,7 +274,7 @@ void FStormSyncTransportServerModule::ExecuteStopServer(const TArray<FString>& A
 	}
 }
 
-void FStormSyncTransportServerModule::ExecuteServerStatus(const TArray<FString>& Args) const
+void FStormSyncTransportServerModule::ExecuteServerStatus(const TArray<FString>& InArgs) const
 {
 	FText StatusText;
 	GetServerStatus(StatusText);
diff --git a/Engine/Plugins/Experimental/StormSync/Source/StormSyncTransportServer/Private/StormSyncTransportServerModule.h b/Engine/Plugins/Experimental/StormSync/Source/StormSyncTransportServer/Private/StormSyncTransportServerModule.h
--- a/Engine/Plugins/Experimental/StormSync/Source/StormSyncTransportServer/Private/StormSyncTransportServerModule.h
+++ b/Engine/Plugins/Experimental/StormSync/Source/StormSyncTransportServer/Private/StormSyncTransportServerModule.h
@@ -16,27 +16,23 @@ public:
 	virtual void ShutdownModule() override;
 
 	//~ Begin IStormSyncTransportClientModule
+	virtual void StartDiscoveryManager() override;
+	virtual void StartServerEndpoint(const FString& InEndpointFriendlyName) override;
 	virtual TSharedPtr<IStormSyncTransportServerLocalEndpoint> CreateServerLocalEndpoint(const FString& InEndpointFriendlyName) const override;
 	virtual FString GetServerEndpointMessageAddressId() const override;
 	virtual FString GetDiscoveryManagerMessageAddressId() const override;
 	virtual FStormSyncHeartbeatEmitter& GetHeartbeatEmitter() const override;
 	virtual bool IsRunning() const override;
-	virtual bool GetServerStatus(FText& StatusText) const override;
+	virtual bool GetServerStatus(FText& OutStatusText) const override;
 	//~ End IStormSyncTransportClientModule
 
 private:
-	/** Our message endpoint provider */
-	TSharedPtr<IStormSyncTransportServerLocalEndpoint, ESPMode::ThreadSafe> ServerEndpoint;
+	/** Starts the heart beat emitter if it was not started already. */
+	void ConditionalStartHeartBeatEmitter();
 
-	/** Unique ptr to our Heartbeat Emitter (Runnable sending heartbeat messages at a fixed interval to subscribed recipients) */
-	TUniquePtr<FStormSyncHeartbeatEmitter> HeartbeatEmitter;
+	/** Starts the discovery manager if it was not started already. */
+	void ConditionalStartDiscoveryManager();
 	
-	/** Unique ptr to our Discovery Manager */
-	TUniquePtr<FStormSyncDiscoveryManager> DiscoveryManager;
-
-	/** References of registered console commands via IConsoleManager */
-	TArray<IConsoleObject*> ConsoleCommands;
-
 	/** Called from StartupModule and sets up console commands for the plugin via IConsoleManager */
 	void RegisterConsoleCommands();
 
@@ -44,17 +40,36 @@ private:
 	void UnregisterConsoleCommands();
 
 	/** Event handler to kick in operations once engine is fully initialized (to publish a connect message) */
-	void OnEngineLoopInitComplete();
+	void OnPostEngineInit();
 
 	/** Notify network this client is ready by sending a connect message and starting heartbeats */
 	void PublishConnectMessage() const;
 
+	/** Publishes the ping message for the server endpoint. */
+	void PublishPingMessage() const;
+
 	/** Command handler for "StormSync.Server.Start" */
-	void ExecuteStartServer(const TArray<FString>& Args);
+	void ExecuteStartServer(const TArray<FString>& InArgs);
 	
 	/** Command handler for "StormSync.Server.Stop" */
-	void ExecuteStopServer(const TArray<FString>& Args);
+	void ExecuteStopServer(const TArray<FString>& InArgs);
 	
 	/** Command handler for "StormSync.Server.Status" */
-	void ExecuteServerStatus(const TArray<FString>& Args) const;
+	void ExecuteServerStatus(const TArray<FString>& InArgs) const;
+
+private:
+	/** Our message endpoint provider */
+	TSharedPtr<IStormSyncTransportServerLocalEndpoint, ESPMode::ThreadSafe> ServerEndpoint;
+
+	/** Unique ptr to our Heartbeat Emitter (Runnable sending heartbeat messages at a fixed interval to subscribed recipients) */
+	TUniquePtr<FStormSyncHeartbeatEmitter> HeartbeatEmitter;
+	
+	/** Unique ptr to our Discovery Manager */
+	TUniquePtr<FStormSyncDiscoveryManager> DiscoveryManager;
+
+	/** References of registered console commands via IConsoleManager */
+	TArray<IConsoleObject*> ConsoleCommands;
+
+	/** Indicates when the engine init is complete. This is used to know when pending messages can be published. */
+	bool bEngineInitComplete = false;
 };
diff --git a/Engine/Plugins/Experimental/StormSync/Source/StormSyncTransportServer/Public/IStormSyncTransportServerModule.h b/Engine/Plugins/Experimental/StormSync/Source/StormSyncTransportServer/Public/IStormSyncTransportServerModule.h
--- a/Engine/Plugins/Experimental/StormSync/Source/StormSyncTransportServer/Public/IStormSyncTransportServerModule.h
+++ b/Engine/Plugins/Experimental/StormSync/Source/StormSyncTransportServer/Public/IStormSyncTransportServerModule.h
@@ -34,7 +34,23 @@ public:
 		return FModuleManager::Get().IsModuleLoaded(ModuleName);
 	}
 
-	/** Create a local transport endpoint */
+	/**
+	 * Starts the discovery manager.
+	 * This is necessary for clients and servers to see each other.
+	 * It must be called even if the server endpoint is not created.
+	 */
+	virtual void StartDiscoveryManager() = 0;
+
+	/**
+	 * Starts a local transport endpoint.
+	 * This will automatically start the discovery manager.
+	 */
+	virtual void StartServerEndpoint(const FString& InEndpointFriendlyName) = 0;
+	
+	/**
+	 * Creates a local transport endpoint.
+	 * @remark Should be used for tests only.
+	 */
 	virtual TSharedPtr<IStormSyncTransportServerLocalEndpoint> CreateServerLocalEndpoint(const FString& InEndpointFriendlyName) const = 0;
 
 	/** Returns Message Address UID for server endpoint if it is currently running, empty string otherwise */
@@ -50,5 +66,5 @@ public:
 	virtual bool IsRunning() const = 0;
 	
 	/** Returns whether Storm Sync Server endpoint is currently active and running, along with a status text indicating current status and endpoint addresses (message bus and tcp server) */
-	virtual bool GetServerStatus(FText& StatusText) const = 0;
+	virtual bool GetServerStatus(FText& OutStatusText) const = 0;
 };
