[Motion Design] Adding "Motion Design" in Storm Sync's description so it shows up in plugin browser search with Motion Design.

diff --git a/Engine/Plugins/Experimental/StormSync/StormSync.uplugin b/Engine/Plugins/Experimental/StormSync/StormSync.uplugin
--- a/Engine/Plugins/Experimental/StormSync/StormSync.uplugin
+++ b/Engine/Plugins/Experimental/StormSync/StormSync.uplugin
@@ -3,7 +3,7 @@
 	"Version": 1,
 	"VersionName": "1.0",
 	"FriendlyName": "Storm Sync",
-	"Description": "Sync, Pull, Push, asset dependencies.",
+	"Description": "Sync, Pull, Push, asset dependencies.\n\nThis plugin is a recommended part of the Motion Design work flow.",
 	"Category": "Experimental",
 	"CreatedBy": "Epic Games, Inc.",
 	"CreatedByURL": "https://epicgames.com",
