Storm Sync: fixed tests failing. Most of these tests involved comparing the error messages with an expected error message. This was changed to instead compare a more reliable enum error code.
- Fixed crash with Storm Sync Archive Factory being used without specifying an outer package.
- Fixed test using path string equality between two relative paths. This was changed to use FPath api to compare paths as these relative paths could be different lexically but still be the same path logically.

diff --git a/Engine/Plugins/Experimental/StormSync/Source/StormSyncDrives/Private/StormSyncDrivesModule.cpp b/Engine/Plugins/Experimental/StormSync/Source/StormSyncDrives/Private/StormSyncDrivesModule.cpp
--- a/Engine/Plugins/Experimental/StormSync/Source/StormSyncDrives/Private/StormSyncDrivesModule.cpp
+++ b/Engine/Plugins/Experimental/StormSync/Source/StormSyncDrives/Private/StormSyncDrivesModule.cpp
@@ -57,13 +57,13 @@ void FStormSyncDrivesModule::ShutdownModule()
 	MountedDrives.Empty();
 }
 
-bool FStormSyncDrivesModule::RegisterMountPoint(const FStormSyncMountPointConfig& InMountPoint, FText& ErrorText)
+bool FStormSyncDrivesModule::RegisterMountPoint(const FStormSyncMountPointConfig& InMountPoint, FText& ErrorText, EStormSyncDriveErrorCode* OutErrorCode)
 {
 	UE_LOG(LogStormSyncDrives, Display, TEXT("FStormSyncDrivesModule::RegisterMountPoint ... MountPoint: %s, Path: %s"), *InMountPoint.MountPoint, *InMountPoint.MountDirectory.Path);
 
 	// Validate first configuration of the mount point
 	FText ValidationError;
-	if (!FStormSyncDrivesUtils::ValidateMountPoint(InMountPoint, ValidationError))
+	if (!FStormSyncDrivesUtils::ValidateMountPoint(InMountPoint, ValidationError, OutErrorCode))
 	{
 		ErrorText = FText::Format(
 			LOCTEXT("ValidationError_MountPoint", "Validation failed for MountPoint (MountPoint: {0}, MountDirectory: {1}) - {2}"),
diff --git a/Engine/Plugins/Experimental/StormSync/Source/StormSyncDrives/Private/StormSyncDrivesModule.h b/Engine/Plugins/Experimental/StormSync/Source/StormSyncDrives/Private/StormSyncDrivesModule.h
--- a/Engine/Plugins/Experimental/StormSync/Source/StormSyncDrives/Private/StormSyncDrivesModule.h
+++ b/Engine/Plugins/Experimental/StormSync/Source/StormSyncDrives/Private/StormSyncDrivesModule.h
@@ -22,7 +22,7 @@ public:
 	//~ End IModuleInterface interface
 
 	//~ Begin IStormSyncDrivesModule interface
-	virtual bool RegisterMountPoint(const FStormSyncMountPointConfig& InMountPoint, FText& ErrorText) override;
+	virtual bool RegisterMountPoint(const FStormSyncMountPointConfig& InMountPoint, FText& ErrorText, EStormSyncDriveErrorCode* OutErrorCode = nullptr) override;
 	virtual bool UnregisterMountPoint(const FStormSyncMountPointConfig& InMountPoint, FText& ErrorText) override;
 	//~ End IStormSyncDrivesModule interface
 
diff --git a/Engine/Plugins/Experimental/StormSync/Source/StormSyncDrives/Private/StormSyncDrivesUtils.cpp b/Engine/Plugins/Experimental/StormSync/Source/StormSyncDrives/Private/StormSyncDrivesUtils.cpp
--- a/Engine/Plugins/Experimental/StormSync/Source/StormSyncDrives/Private/StormSyncDrivesUtils.cpp
+++ b/Engine/Plugins/Experimental/StormSync/Source/StormSyncDrives/Private/StormSyncDrivesUtils.cpp
@@ -3,18 +3,20 @@
 #include "StormSyncDrivesUtils.h"
 
 #include "HAL/FileManager.h"
+#include "Misc/ScopeExit.h"
+#include "StormSyncDrivesEnums.h"
 #include "StormSyncDrivesSettings.h"
 
 #define LOCTEXT_NAMESPACE "FStormSyncDrivesUtils"
 
-bool FStormSyncDrivesUtils::ValidateMountPoint(const FStormSyncMountPointConfig& InMountPoint, FText& ErrorText)
+bool FStormSyncDrivesUtils::ValidateMountPoint(const FStormSyncMountPointConfig& InMountPoint, FText& ErrorText, EStormSyncDriveErrorCode* OutErrorCode)
 {
-	if (!ValidateRootPath(InMountPoint.MountPoint, ErrorText))
+	if (!ValidateRootPath(InMountPoint.MountPoint, ErrorText, OutErrorCode))
 	{
 		return false;
 	}
 
-	if (!ValidateDirectory(InMountPoint.MountDirectory, ErrorText))
+	if (!ValidateDirectory(InMountPoint.MountDirectory, ErrorText, OutErrorCode))
 	{
 		return false;
 	}
@@ -22,12 +24,22 @@ bool FStormSyncDrivesUtils::ValidateMountPoint(const FStormSyncMountPointConfig&
 	return true;
 }
 
-bool FStormSyncDrivesUtils::ValidateRootPath(const FString& InRootPath, FText& ErrorText)
+bool FStormSyncDrivesUtils::ValidateRootPath(const FString& InRootPath, FText& ErrorText, EStormSyncDriveErrorCode* OutErrorCode)
 {
+	EStormSyncDriveErrorCode ErrorCode = EStormSyncDriveErrorCode::None;
+	ON_SCOPE_EXIT
+	{
+		if (OutErrorCode)
+		{
+			*OutErrorCode = ErrorCode;
+		}
+	};
+
 	// Non empty
 	if (InRootPath.IsEmpty())
 	{
 		ErrorText = LOCTEXT("ValidationError_Empty_RootPath", "Root Path is empty");
+		ErrorCode = EStormSyncDriveErrorCode::EmptyRootPath;
 		return false;
 	}
 
@@ -35,6 +47,7 @@ bool FStormSyncDrivesUtils::ValidateRootPath(const FString& InRootPath, FText& E
 	if (InRootPath == TEXT("/"))
 	{
 		ErrorText = LOCTEXT("ValidationError_SingleSlash", "Cannot mount a directory on the \"/\" Root Path");
+		ErrorCode = EStormSyncDriveErrorCode::TopLevelRootDisallowed;
 		return false;
 	}
 
@@ -43,6 +56,7 @@ bool FStormSyncDrivesUtils::ValidateRootPath(const FString& InRootPath, FText& E
 	if (!IsValidTextForLongPackageName(InRootPath, &Reason))
 	{
 		ErrorText = FPackageName::FormatErrorAsText(InRootPath, Reason);
+		ErrorCode = EStormSyncDriveErrorCode::InvalidPackageNameFormat;
 		return false;
 	}
 
@@ -57,23 +71,35 @@ bool FStormSyncDrivesUtils::ValidateRootPath(const FString& InRootPath, FText& E
 			LOCTEXT("ValidationError_RootPath_Level", "\"{0}\" is not a valid root path. It must be one level only (eg. /Example) when not within /Game."),
 			FText::FromString(InRootPath)
 		);
+		ErrorCode = EStormSyncDriveErrorCode::InvalidPathLevel;
 		return false;
 	}
 
 	return true;
 }
 
-bool FStormSyncDrivesUtils::ValidateDirectory(const FDirectoryPath& InDirectory, FText& ErrorText)
+bool FStormSyncDrivesUtils::ValidateDirectory(const FDirectoryPath& InDirectory, FText& ErrorText, EStormSyncDriveErrorCode* OutErrorCode)
 {
+	EStormSyncDriveErrorCode ErrorCode = EStormSyncDriveErrorCode::None;
+	ON_SCOPE_EXIT
+	{
+		if (OutErrorCode)
+		{
+			*OutErrorCode = ErrorCode;
+		}
+	};
+
 	if (InDirectory.Path.IsEmpty())
 	{
 		ErrorText = LOCTEXT("ValidationError_Empty_Directory", "Directory Path is empty");
+		ErrorCode = EStormSyncDriveErrorCode::EmptyDirectoryPath;
 		return false;
 	}
 
 	if (!IFileManager::Get().DirectoryExists(*InDirectory.Path))
 	{
 		ErrorText = FText::Format(LOCTEXT("ValidationError_Invalid_Directory", "Directory '{0}' does not exist."), FText::FromString(InDirectory.Path));
+		ErrorCode = EStormSyncDriveErrorCode::DirectoryNotFound;
 		return false;
 	}
 
diff --git a/Engine/Plugins/Experimental/StormSync/Source/StormSyncDrives/Private/Tests/StormSyncDrivesModule.spec.cpp b/Engine/Plugins/Experimental/StormSync/Source/StormSyncDrives/Private/Tests/StormSyncDrivesModule.spec.cpp
--- a/Engine/Plugins/Experimental/StormSync/Source/StormSyncDrives/Private/Tests/StormSyncDrivesModule.spec.cpp
+++ b/Engine/Plugins/Experimental/StormSync/Source/StormSyncDrives/Private/Tests/StormSyncDrivesModule.spec.cpp
@@ -3,8 +3,9 @@
 #include "HAL/FileManager.h"
 #include "IStormSyncDrivesModule.h"
 #include "Misc/AutomationTest.h"
-#include "Misc/Paths.h"
 #include "Misc/PackageName.h"
+#include "Misc/Paths.h"
+#include "StormSyncDrivesEnums.h"
 #include "StormSyncDrivesSettings.h"
 
 BEGIN_DEFINE_SPEC(FStormSyncDrivesModuleSpec, "StormSync.StormSyncDrives", EAutomationTestFlags::ProductFilter | EAutomationTestFlags_ApplicationContextMask)
@@ -43,12 +44,14 @@ void FStormSyncDrivesModuleSpec::Define()
 
 				// Two times, once for message log, once for error log
 				const FString ExpectedErrorText = TEXT("Root Path is empty");
-				AddExpectedError(ExpectedErrorText, EAutomationExpectedErrorFlags::Contains, 2);
+				AddExpectedError(ExpectedErrorText, EAutomationExpectedErrorFlags::Contains, 1);
 
 				FText ErrorText;
-				const bool bSuccess = IStormSyncDrivesModule::Get().RegisterMountPoint(MountPoint, ErrorText);
+				EStormSyncDriveErrorCode ErrorCode = EStormSyncDriveErrorCode::None;
+
+				const bool bSuccess = IStormSyncDrivesModule::Get().RegisterMountPoint(MountPoint, ErrorText, &ErrorCode);
 				TestFalse(TEXT("Expected return value is false"), bSuccess);
-				TestEqual(TEXT("ErrorText expected value"), ErrorText.ToString(), ExpectedErrorText);
+				TestTrue(TEXT("ErrorText expected value"), ErrorCode == EStormSyncDriveErrorCode::EmptyRootPath);
 			});
 
 			It(TEXT("should fail on invalid input - empty directory"), [this]()
@@ -58,12 +61,14 @@ void FStormSyncDrivesModuleSpec::Define()
 				MountPoint.MountDirectory.Path = TEXT("");
 
 				const FString ExpectedErrorText = TEXT("Directory Path is empty");
-				AddExpectedError(ExpectedErrorText, EAutomationExpectedErrorFlags::Contains, 2);
+				AddExpectedError(ExpectedErrorText, EAutomationExpectedErrorFlags::Contains, 1);
 
 				FText ErrorText;
-				const bool bSuccess = IStormSyncDrivesModule::Get().RegisterMountPoint(MountPoint, ErrorText);
+				EStormSyncDriveErrorCode ErrorCode = EStormSyncDriveErrorCode::None;
+
+				const bool bSuccess = IStormSyncDrivesModule::Get().RegisterMountPoint(MountPoint, ErrorText, &ErrorCode);
 				TestFalse(TEXT("Expected return value is false"), bSuccess);
-				TestEqual(TEXT("ErrorText expected value"), ErrorText.ToString(), ExpectedErrorText);
+				TestTrue(TEXT("ErrorText expected value"), ErrorCode == EStormSyncDriveErrorCode::EmptyDirectoryPath);
 			});
 
 			It(TEXT("should fail on invalid input - invalid directory"), [this]()
@@ -73,12 +78,14 @@ void FStormSyncDrivesModuleSpec::Define()
 				MountPoint.MountDirectory.Path = InvalidDirectoryPath;
 
 				const FString ExpectedErrorText = FString::Printf(TEXT("Directory '%s' does not exist."), *InvalidDirectoryPath);
-				AddExpectedError(ExpectedErrorText, EAutomationExpectedErrorFlags::Contains, 2);
+				AddExpectedError(ExpectedErrorText, EAutomationExpectedErrorFlags::Contains, 1);
 
 				FText ErrorText;
-				const bool bSuccess = IStormSyncDrivesModule::Get().RegisterMountPoint(MountPoint, ErrorText);
+				EStormSyncDriveErrorCode ErrorCode = EStormSyncDriveErrorCode::None;
+
+				const bool bSuccess = IStormSyncDrivesModule::Get().RegisterMountPoint(MountPoint, ErrorText, &ErrorCode);
 				TestFalse(TEXT("Expected return value is false"), bSuccess);
-				TestEqual(TEXT("ErrorText expected value"), ErrorText.ToString(), ExpectedErrorText);
+				TestTrue(TEXT("ErrorText expected value"), ErrorCode == EStormSyncDriveErrorCode::DirectoryNotFound);
 			});
 
 			It(TEXT("should mount directory with valid configuration"), [this]()
diff --git a/Engine/Plugins/Experimental/StormSync/Source/StormSyncDrives/Public/IStormSyncDrivesModule.h b/Engine/Plugins/Experimental/StormSync/Source/StormSyncDrives/Public/IStormSyncDrivesModule.h
--- a/Engine/Plugins/Experimental/StormSync/Source/StormSyncDrives/Public/IStormSyncDrivesModule.h
+++ b/Engine/Plugins/Experimental/StormSync/Source/StormSyncDrives/Public/IStormSyncDrivesModule.h
@@ -5,6 +5,7 @@
 #include "Modules/ModuleInterface.h"
 #include "Modules/ModuleManager.h"
 
+enum class EStormSyncDriveErrorCode : uint8;
 struct FStormSyncMountPointConfig;
 
 /** Public interface for Storm Sync Drives module and API */
@@ -30,10 +31,10 @@ public:
 	 *
 	 * @param InMountPoint The config with root mount path & directory to mount
 	 * @param ErrorText Localized text describing the error in case mounting failed
-	 *
+	 * @param OutError optional enum for further details on the operation's error if any
 	 * @returns true if we were able to register the mount point
 	 */
-	virtual bool RegisterMountPoint(const FStormSyncMountPointConfig& InMountPoint, FText& ErrorText) = 0;
+	virtual bool RegisterMountPoint(const FStormSyncMountPointConfig& InMountPoint, FText& ErrorText, EStormSyncDriveErrorCode* OutError = nullptr) = 0;
 
 	/**
 	 * Unregisters the provided mount point config, if it was previously mounted.
diff --git a/Engine/Plugins/Experimental/StormSync/Source/StormSyncDrives/Public/StormSyncDrivesEnums.h b/Engine/Plugins/Experimental/StormSync/Source/StormSyncDrives/Public/StormSyncDrivesEnums.h
new file mode 100644
--- /dev/null
+++ b/Engine/Plugins/Experimental/StormSync/Source/StormSyncDrives/Public/StormSyncDrivesEnums.h
@@ -0,0 +1,29 @@
+// Copyright Epic Games, Inc. All Rights Reserved.
+
+#pragma once
+
+#include "HAL/Platform.h"
+
+enum class EStormSyncDriveErrorCode : uint8
+{
+	/** No error */
+	None,
+
+	/** Root path is empty */
+	EmptyRootPath,
+
+	/** Path is "/". No mounting allowed at the top level root */
+	TopLevelRootDisallowed,
+
+	/** Path does not have a standard package name format */
+	InvalidPackageNameFormat,
+
+	/** Path is expected to be within /Game or one level path e.g. /Example when not within /Game */
+	InvalidPathLevel,
+
+	/** Directory path is empty */
+	EmptyDirectoryPath,
+
+	/** Directory path does not exist */
+	DirectoryNotFound,
+};
diff --git a/Engine/Plugins/Experimental/StormSync/Source/StormSyncDrives/Public/StormSyncDrivesUtils.h b/Engine/Plugins/Experimental/StormSync/Source/StormSyncDrives/Public/StormSyncDrivesUtils.h
--- a/Engine/Plugins/Experimental/StormSync/Source/StormSyncDrives/Public/StormSyncDrivesUtils.h
+++ b/Engine/Plugins/Experimental/StormSync/Source/StormSyncDrives/Public/StormSyncDrivesUtils.h
@@ -5,6 +5,7 @@
 #include "Engine/EngineTypes.h"
 #include "Misc/PackageName.h"
 
+enum class EStormSyncDriveErrorCode : uint8;
 struct FStormSyncMountPointConfig;
 
 /** Provides validation helpers */
@@ -12,13 +13,13 @@ class STORMSYNCDRIVES_API FStormSyncDrivesUtils
 {
 public:
 	/** Checks validity for a mount point */
-	static bool ValidateMountPoint(const FStormSyncMountPointConfig& InMountPoint, FText& ErrorText);
+	static bool ValidateMountPoint(const FStormSyncMountPointConfig& InMountPoint, FText& ErrorText, EStormSyncDriveErrorCode* OutErrorCode = nullptr);
 
 	/** Checks validity for a mount point */
-	static bool ValidateRootPath(const FString& InRootPath, FText& ErrorText);
+	static bool ValidateRootPath(const FString& InRootPath, FText& ErrorText, EStormSyncDriveErrorCode* OutErrorCode = nullptr);
 
 	/** Checks validity for a mount directory */
-	static bool ValidateDirectory(const FDirectoryPath& InDirectory, FText& ErrorText);
+	static bool ValidateDirectory(const FDirectoryPath& InDirectory, FText& ErrorText, EStormSyncDriveErrorCode* OutErrorCode = nullptr);
 
 	/** Checks validity for any duplicates regarding Mount Point (Root Paths) and filesystem directory */
 	static bool ValidateNonDuplicates(const TArray<FStormSyncMountPointConfig> InMountPoints, TArray<FText>& ValidationErrors);
diff --git a/Engine/Plugins/Experimental/StormSync/Source/StormSyncEditor/Private/Tests/StormSyncArchiveFactory.spec.cpp b/Engine/Plugins/Experimental/StormSync/Source/StormSyncEditor/Private/Tests/StormSyncArchiveFactory.spec.cpp
--- a/Engine/Plugins/Experimental/StormSync/Source/StormSyncEditor/Private/Tests/StormSyncArchiveFactory.spec.cpp
+++ b/Engine/Plugins/Experimental/StormSync/Source/StormSyncEditor/Private/Tests/StormSyncArchiveFactory.spec.cpp
@@ -2,6 +2,7 @@
 
 #include "Factories/StormSyncArchiveFactory.h"
 #include "Misc/AutomationTest.h"
+#include "UObject/Package.h"
 
 BEGIN_DEFINE_SPEC(FStormSyncArchiveFactorySpec, "StormSync.StormSyncEditor.StormSyncArchiveFactory", EAutomationTestFlags::ProductFilter | EAutomationTestFlags_ApplicationContextMask)
 
@@ -67,8 +68,10 @@ void FStormSyncArchiveFactorySpec::Define()
 			AddExpectedError(TEXT("Failed to read file 'invalid file' error"), EAutomationExpectedErrorFlags::Contains, 1);
 			AddExpectedError(TEXT("Failed to load file 'invalid file' to array"), EAutomationExpectedErrorFlags::Contains, 1);
 
+			UPackage* const Package = GetTransientPackage();
+
 			bool bCanceled = false;
-			const UObject* Result = Factory->FactoryCreateFile(UStormSyncArchiveData::StaticClass(), nullptr, NAME_None, RF_Public | RF_Standalone, Filename, nullptr, GWarn, bCanceled);
+			const UObject* Result = Factory->FactoryCreateFile(UStormSyncArchiveData::StaticClass(), Package, NAME_None, RF_Public | RF_Standalone, Filename, nullptr, GWarn, bCanceled);
 			TestNotNull(TEXT("Result is StormSyncArchiveData object"), Result);
 			TestTrue(TEXT("Result is StormSyncArchiveData object"), Result->IsA(UStormSyncArchiveData::StaticClass()));
 			TestEqual(TEXT("Result.Filename"), Cast<UStormSyncArchiveData>(Result)->Filename, Filename);
diff --git a/Engine/Plugins/Experimental/StormSync/Source/StormSyncTests/Private/StormSyncPackageDescriptor.spec.cpp b/Engine/Plugins/Experimental/StormSync/Source/StormSyncTests/Private/StormSyncPackageDescriptor.spec.cpp
--- a/Engine/Plugins/Experimental/StormSync/Source/StormSyncTests/Private/StormSyncPackageDescriptor.spec.cpp
+++ b/Engine/Plugins/Experimental/StormSync/Source/StormSyncTests/Private/StormSyncPackageDescriptor.spec.cpp
@@ -84,7 +84,7 @@ void FStormSyncPackageDescriptorSpec::Define()
 					AddError(ErrorText.ToString());
 				}
 
-				TestEqual(FString::Printf(TEXT("Has expected filepath for %s"), *PackageName.ToString()), DestFilepath, ExpectedFilepath);
+				TestTrue(FString::Printf(TEXT("Has expected filepath for %s"), *PackageName.ToString()), FPaths::IsSamePath(DestFilepath, ExpectedFilepath));
 			}
 		});
 
