Fixed some more printf mismatches.

diff --git a/Engine/Plugins/Editor/AssetSearch/Source/Private/AssetSearchDatabase.cpp b/Engine/Plugins/Editor/AssetSearch/Source/Private/AssetSearchDatabase.cpp
--- a/Engine/Plugins/Editor/AssetSearch/Source/Private/AssetSearchDatabase.cpp
+++ b/Engine/Plugins/Editor/AssetSearch/Source/Private/AssetSearchDatabase.cpp
@@ -720,7 +720,7 @@ FString FAssetSearchDatabase::GetLastError() const
 
 void FAssetSearchDatabase::LogLastError() const
 {
-	UE_LOG(LogAssetSearch, Error, TEXT("Database Error: %s"), *SessionPath, *GetLastError());
+	UE_LOG(LogAssetSearch, Error, TEXT("Database '%s' Error: %s"), *SessionPath, *GetLastError());
 }
 
 bool FAssetSearchDatabase::IsAssetUpToDate(const FAssetData& InAssetData, const FString& IndexedJsonHash)
diff --git a/Engine/Plugins/Experimental/StormSync/Source/StormSyncEditor/Private/Factories/StormSyncArchiveFactory.cpp b/Engine/Plugins/Experimental/StormSync/Source/StormSyncEditor/Private/Factories/StormSyncArchiveFactory.cpp
--- a/Engine/Plugins/Experimental/StormSync/Source/StormSyncEditor/Private/Factories/StormSyncArchiveFactory.cpp
+++ b/Engine/Plugins/Experimental/StormSync/Source/StormSyncEditor/Private/Factories/StormSyncArchiveFactory.cpp
@@ -38,7 +38,7 @@ bool UStormSyncArchiveFactory::FactoryCanImport(const FString& Filename)
 
 UObject* UStormSyncArchiveFactory::ImportObject(UClass* InClass, UObject* InOuter, FName InName, EObjectFlags Flags, const FString& Filename, const TCHAR* Params, bool& OutCanceled)
 {
-	UE_LOG(LogStormSyncEditor, Display, TEXT("UStormSyncArchiveFactory::ImportObject InClass: %s, InOuter: %s, Filename: %s"), *GetNameSafe(InClass), *GetNameSafe(InOuter), *InName.ToString(), *Filename);
+	UE_LOG(LogStormSyncEditor, Display, TEXT("UStormSyncArchiveFactory::ImportObject InClass: %s, InOuter: %s, InName: %s, Filename: %s"), *GetNameSafe(InClass), *GetNameSafe(InOuter), *InName.ToString(), *Filename);
 
 	// Import Object override to customize how import is done and account for this factory operating on a dummy imported object.
 
