Storm Sync: fixed issue with no reimport handler to deal with reimport logic when dragging the same files again

diff --git a/Engine/Plugins/Experimental/StormSync/Source/StormSyncEditor/Private/Factories/StormSyncArchiveFactory.cpp b/Engine/Plugins/Experimental/StormSync/Source/StormSyncEditor/Private/Factories/StormSyncArchiveFactory.cpp
--- a/Engine/Plugins/Experimental/StormSync/Source/StormSyncEditor/Private/Factories/StormSyncArchiveFactory.cpp
+++ b/Engine/Plugins/Experimental/StormSync/Source/StormSyncEditor/Private/Factories/StormSyncArchiveFactory.cpp
@@ -104,6 +104,43 @@ UObject* UStormSyncArchiveFactory::FactoryCreateFile(UClass* InClass, UObject* I
 	return ArchiveData;
 }
 
+bool UStormSyncArchiveFactory::CanReimport(UObject* InObject, TArray<FString>& OutFilenames)
+{
+	if (const UStormSyncArchiveData* ArchiveData = Cast<UStormSyncArchiveData>(InObject))
+	{
+		OutFilenames.Add(ArchiveData->Filename);
+		return true;
+	}
+	return false;
+}
+
+void UStormSyncArchiveFactory::SetReimportPaths(UObject* InObject, const TArray<FString>& InNewReimportPaths)
+{
+	UStormSyncArchiveData* ArchiveData = Cast<UStormSyncArchiveData>(InObject);
+	if (ArchiveData && ensure(InNewReimportPaths.Num() == 1))
+	{
+		ArchiveData->Filename = InNewReimportPaths[0];
+	}
+}
+
+EReimportResult::Type UStormSyncArchiveFactory::Reimport(UObject* InObject)
+{
+	if (UStormSyncArchiveData* ArchiveData = Cast<UStormSyncArchiveData>(InObject))
+	{
+		PendingImports.Add(ArchiveData->Filename, FSoftObjectPath(ArchiveData));
+
+		// Broadcast import event and delay import until next tick to avoid blocking the process that files were dragged from
+		// (delay handled internally by Import subsystem)
+		FStormSyncCoreDelegates::OnRequestImportFile.Broadcast(ArchiveData->Filename);
+
+		// This does not look very nice, but returning cancelled means there's no success/fail pop-up notification
+		// that will show even before the import is fully processed.
+		// @see FReimportManager::ReimportAsync
+		return EReimportResult::Cancelled;
+	}
+	return EReimportResult::Failed;
+}
+
 void UStormSyncArchiveFactory::OnFileImported(const FString& InFilename)
 {
 	if (const FSoftObjectPath* FoundAssetPath = PendingImports.Find(InFilename))
diff --git a/Engine/Plugins/Experimental/StormSync/Source/StormSyncEditor/Private/Factories/StormSyncArchiveFactory.h b/Engine/Plugins/Experimental/StormSync/Source/StormSyncEditor/Private/Factories/StormSyncArchiveFactory.h
--- a/Engine/Plugins/Experimental/StormSync/Source/StormSyncEditor/Private/Factories/StormSyncArchiveFactory.h
+++ b/Engine/Plugins/Experimental/StormSync/Source/StormSyncEditor/Private/Factories/StormSyncArchiveFactory.h
@@ -2,6 +2,7 @@
 
 #pragma once
 
+#include "EditorReimportHandler.h"
 #include "Factories/Factory.h"
 #include "StormSyncArchiveFactory.generated.h"
 
@@ -27,7 +28,7 @@ public:
  * file it is importing, but more like a "zip" including all assets and their inner dependencies.
  */
 UCLASS(hidecategories=Object, MinimalAPI)
-class UStormSyncArchiveFactory : public UFactory
+class UStormSyncArchiveFactory : public UFactory, public FReimportHandler
 {
 	GENERATED_BODY()
 public:
@@ -40,6 +41,12 @@ public:
 	virtual UObject* FactoryCreateFile(UClass* InClass, UObject* InParent, FName InName, EObjectFlags Flags, const FString& Filename, const TCHAR* Params, FFeedbackContext* Warn, bool& bOutOperationCanceled) override;
 	//~ End UFactory interface
 
+	//~ Begin FReimportHandler
+	virtual bool CanReimport(UObject* InObject, TArray<FString>& OutFilenames) override;
+	virtual void SetReimportPaths(UObject* InObject, const TArray<FString>& InNewReimportPaths) override;
+	virtual EReimportResult::Type Reimport(UObject* InObject) override;
+	//~ End FReimportHandler
+
 	void OnFileImported(const FString& InFilename);
 
 	TMap<FString, FSoftObjectPath> PendingImports;
