Storm Sync: switched cancelled return value to success

diff --git a/Engine/Plugins/Experimental/StormSync/Source/StormSyncEditor/Private/Factories/StormSyncArchiveFactory.cpp b/Engine/Plugins/Experimental/StormSync/Source/StormSyncEditor/Private/Factories/StormSyncArchiveFactory.cpp
--- a/Engine/Plugins/Experimental/StormSync/Source/StormSyncEditor/Private/Factories/StormSyncArchiveFactory.cpp
+++ b/Engine/Plugins/Experimental/StormSync/Source/StormSyncEditor/Private/Factories/StormSyncArchiveFactory.cpp
@@ -132,11 +132,7 @@ EReimportResult::Type UStormSyncArchiveFactory::Reimport(UObject* InObject)
 		// Broadcast import event and delay import until next tick to avoid blocking the process that files were dragged from
 		// (delay handled internally by Import subsystem)
 		FStormSyncCoreDelegates::OnRequestImportFile.Broadcast(ArchiveData->Filename);
-
-		// This does not look very nice, but returning cancelled means there's no success/fail pop-up notification
-		// that will show even before the import is fully processed.
-		// @see FReimportManager::ReimportAsync
-		return EReimportResult::Cancelled;
+		return EReimportResult::Succeeded;
 	}
 	return EReimportResult::Failed;
 }
