[StormSync] Revive Connetion fix

Problem:
StormSync has a peer to peer discovery service with heartbeats. Each StormSync node sends heartbeats to other nodes. If a node becomes non-responsive for long enough, the other nodes are going to remove it from active connections.
The issue occurs when a non-responsive node resume sending heartbeats. This falls in the "revive connection code path" and there was a bug in that. In FStormSyncDiscoveryManager::HandleHeartbeatMessage, FStormSyncCoreDelegates::OnServiceDiscoveryConnection was being called with the LOCAL info, when receiving a heart beat from a remote system, but registering ITSELF as the connection. That's why from that point on, pushing/syncing assets to that node would always say "synced" because it is actually syncing with itself. And there is no way to recover unless restarting the system.

Fix:
RegisterConnection must be called with a FStormSyncTransportConnectMessage from the remote connection. This is now only done when receiving the message from the remote. It registers the connection when receiving the heartbeat still, but postpone broadcasting the connecrtion info until it actually receives it.

Extra:
When a node is non-responsive for a while, it causes timeouts in the other endpoints and they need to be woken up.
Adding a new "StormSync.Discovery.Wakeup" command to be able to reactivate client end points when they become inactive.
This is done manually from the editing station for now. This is a work around until we figure out a better solution.

diff --git a/Engine/Plugins/Experimental/StormSync/Source/StormSyncCore/Private/StormSyncCoreDelegates.cpp b/Engine/Plugins/Experimental/StormSync/Source/StormSyncCore/Private/StormSyncCoreDelegates.cpp
--- a/Engine/Plugins/Experimental/StormSync/Source/StormSyncCore/Private/StormSyncCoreDelegates.cpp
+++ b/Engine/Plugins/Experimental/StormSync/Source/StormSyncCore/Private/StormSyncCoreDelegates.cpp
@@ -12,6 +12,7 @@ FStormSyncCoreDelegates::FOnReceivingBytes FStormSyncCoreDelegates::OnReceivingB
 FStormSyncCoreDelegates::FOnTransferComplete FStormSyncCoreDelegates::OnTransferComplete;
 FStormSyncCoreDelegates::FOnServiceDiscoveryConnection FStormSyncCoreDelegates::OnServiceDiscoveryConnection;
 FStormSyncCoreDelegates::FOnServiceDiscoveryStateChange FStormSyncCoreDelegates::OnServiceDiscoveryStateChange;
+FStormSyncCoreDelegates::FOnServiceDiscoveryReceivedWakeup FStormSyncCoreDelegates::OnServiceDiscoveryReceivedWakeup;
 FStormSyncCoreDelegates::FOnServiceDiscoveryServerStatusChange FStormSyncCoreDelegates::OnServiceDiscoveryServerStatusChange;
 FStormSyncCoreDelegates::FOnServiceDiscoveryDisconnection FStormSyncCoreDelegates::OnServiceDiscoveryDisconnection;
 FStormSyncCoreDelegates::FOnStormSyncServerStarted FStormSyncCoreDelegates::OnStormSyncServerStarted;
diff --git a/Engine/Plugins/Experimental/StormSync/Source/StormSyncCore/Public/StormSyncCoreDelegates.h b/Engine/Plugins/Experimental/StormSync/Source/StormSyncCore/Public/StormSyncCoreDelegates.h
--- a/Engine/Plugins/Experimental/StormSync/Source/StormSyncCore/Public/StormSyncCoreDelegates.h
+++ b/Engine/Plugins/Experimental/StormSync/Source/StormSyncCore/Public/StormSyncCoreDelegates.h
@@ -150,6 +150,12 @@ struct STORMSYNCCORE_API FStormSyncCoreDelegates
 	 */
 	DECLARE_MULTICAST_DELEGATE_TwoParams(FOnServiceDiscoveryStateChange, const FString&, EStormSyncConnectedDeviceState);
 
+
+	/**
+	 * Wakeup event sent through the discovery service to re-active the other services.
+	 */
+	DECLARE_MULTICAST_DELEGATE(FOnServiceDiscoveryReceivedWakeup);
+
 	/**
 	 * Delegate type for when server status for a connection on storm sync network has changed (either running or stopped).
 	 *
@@ -214,6 +220,11 @@ struct STORMSYNCCORE_API FStormSyncCoreDelegates
 	 */
 	static FOnServiceDiscoveryStateChange OnServiceDiscoveryStateChange;
 
+	/**
+	 * Called when the discovery service receives a wakeup request.
+	 */
+	static FOnServiceDiscoveryReceivedWakeup OnServiceDiscoveryReceivedWakeup;
+
 	/** Called when service discovery manager detects a state change for remote server endpoint on a connected device */
 	static FOnServiceDiscoveryServerStatusChange OnServiceDiscoveryServerStatusChange;
 
diff --git a/Engine/Plugins/Experimental/StormSync/Source/StormSyncTransportClient/Private/StormSyncTransportClientModule.cpp b/Engine/Plugins/Experimental/StormSync/Source/StormSyncTransportClient/Private/StormSyncTransportClientModule.cpp
--- a/Engine/Plugins/Experimental/StormSync/Source/StormSyncTransportClient/Private/StormSyncTransportClientModule.cpp
+++ b/Engine/Plugins/Experimental/StormSync/Source/StormSyncTransportClient/Private/StormSyncTransportClientModule.cpp
@@ -7,6 +7,7 @@
 #include "IStormSyncTransportLocalEndpoint.h"
 #include "MessageEndpoint.h"
 #include "Misc/CoreDelegates.h"
+#include "StormSyncCoreDelegates.h"
 #include "StormSyncTransportClientEndpoint.h"
 #include "StormSyncTransportClientLog.h"
 #include "StormSyncTransportMessages.h"
@@ -25,11 +26,14 @@ void FStormSyncTransportClientModule::StartupModule()
 	
 	// Register for engine initialization completed so we can broadcast presence over the network from this client to other clients
 	FCoreDelegates::OnPostEngineInit.AddRaw(this, &FStormSyncTransportClientModule::OnPostEngineInit);
+
+	FStormSyncCoreDelegates::OnServiceDiscoveryReceivedWakeup.AddRaw(this, &FStormSyncTransportClientModule::OnWakeup);
 }
 
 void FStormSyncTransportClientModule::ShutdownModule()
 {
 	FCoreDelegates::OnPostEngineInit.RemoveAll(this);
+	FStormSyncCoreDelegates::OnServiceDiscoveryReceivedWakeup.RemoveAll(this);
 	
 	if (ClientEndpoint.IsValid())
 	{
@@ -239,6 +243,17 @@ void FStormSyncTransportClientModule::PublishStatusPingMessage() const
 	MessageEndpoint->Publish(Message);
 }
 
+void FStormSyncTransportClientModule::OnWakeup()
+{
+	const FMessageEndpointSharedPtr MessageEndpoint = GetClientMessageEndpoint();
+	if (!MessageEndpoint.IsValid())
+	{
+		UE_LOG(LogStormSyncClient, Warning, TEXT("FStormSyncTransportClientModule::OnWakeup - Unable to get client message endpoint"));
+		return;
+	}
+	MessageEndpoint->Publish(FMessageEndpoint::MakeMessage<FStormSyncTransportPingMessage>());
+}
+
 void FStormSyncTransportClientModule::ExecuteStartClient(const TArray<FString>& InArgs)
 {
 	StartClientEndpoint(InArgs.Num() ? InArgs[0] : DefaultClientEndpointName);	
diff --git a/Engine/Plugins/Experimental/StormSync/Source/StormSyncTransportClient/Private/StormSyncTransportClientModule.h b/Engine/Plugins/Experimental/StormSync/Source/StormSyncTransportClient/Private/StormSyncTransportClientModule.h
--- a/Engine/Plugins/Experimental/StormSync/Source/StormSyncTransportClient/Private/StormSyncTransportClientModule.h
+++ b/Engine/Plugins/Experimental/StormSync/Source/StormSyncTransportClient/Private/StormSyncTransportClientModule.h
@@ -42,6 +42,9 @@ private:
 	/** Publish the client endpoint's ping message. */
 	void PublishStatusPingMessage() const;
 
+	/** Handler for wake up event. */
+	void OnWakeup();
+
 	/** Command handler for starting the client. */
 	void ExecuteStartClient(const TArray<FString>& InArgs);
 	
diff --git a/Engine/Plugins/Experimental/StormSync/Source/StormSyncTransportCore/Public/StormSyncTransportMessages.h b/Engine/Plugins/Experimental/StormSync/Source/StormSyncTransportCore/Public/StormSyncTransportMessages.h
--- a/Engine/Plugins/Experimental/StormSync/Source/StormSyncTransportCore/Public/StormSyncTransportMessages.h
+++ b/Engine/Plugins/Experimental/StormSync/Source/StormSyncTransportCore/Public/StormSyncTransportMessages.h
@@ -80,6 +80,13 @@ struct FStormSyncTransportConnectMessage : public FStormSyncConnectionInfo
 	GENERATED_BODY()
 };
 
+/** Request sent by the discovery service to wake up the other endpoints in case they become inactive. */
+USTRUCT()
+struct FStormSyncTransportWakeupRequest
+{
+	GENERATED_BODY()
+};
+
 /**
  * FStormSyncTransportSyncRequest.
  *
diff --git a/Engine/Plugins/Experimental/StormSync/Source/StormSyncTransportServer/Private/ServiceDiscovery/StormSyncDiscoveryManager.cpp b/Engine/Plugins/Experimental/StormSync/Source/StormSyncTransportServer/Private/ServiceDiscovery/StormSyncDiscoveryManager.cpp
--- a/Engine/Plugins/Experimental/StormSync/Source/StormSyncTransportServer/Private/ServiceDiscovery/StormSyncDiscoveryManager.cpp
+++ b/Engine/Plugins/Experimental/StormSync/Source/StormSyncTransportServer/Private/ServiceDiscovery/StormSyncDiscoveryManager.cpp
@@ -23,7 +23,8 @@ FStormSyncDiscoveryManager::FStormSyncDiscoveryManager(double InHeartbeatTimeout
 	MessageEndpoint = FMessageEndpoint::Builder(TEXT("StormSyncMessageHeartbeatManager (StormSyncDiscoveryManager)"))
 		.ReceivingOnThread(ENamedThreads::GameThread)
 		.Handling<FStormSyncTransportConnectMessage>(this, &FStormSyncDiscoveryManager::HandleConnectMessage)
-		.Handling<FStormSyncTransportHeartbeatMessage>(this, &FStormSyncDiscoveryManager::HandleHeartbeatMessage);
+		.Handling<FStormSyncTransportHeartbeatMessage>(this, &FStormSyncDiscoveryManager::HandleHeartbeatMessage)
+		.Handling<FStormSyncTransportWakeupRequest>(this, &FStormSyncDiscoveryManager::HandleWakeupMessage);
 
 	bRunning = MessageEndpoint.IsValid();
 	if (bRunning)
@@ -166,6 +167,26 @@ void FStormSyncDiscoveryManager::PublishConnectMessage()
 	}
 }
 
+void FStormSyncDiscoveryManager::SendWakeUp()
+{
+	if (!MessageEndpoint.IsValid())
+	{
+		UE_LOG(LogStormSyncServer, Error, TEXT("FStormSyncDiscoveryManager::SendWakeUp - Unable to send Connect Message cause Message Endpoint is invalid"));
+		return;
+	}
+
+	UE_LOG(LogStormSyncServer, VeryVerbose, TEXT("FStormSyncTransportClientModule::SendWakeUp - Sending Wakeup Request to all connections."));
+	FStormSyncTransportWakeupRequest* Message = FMessageEndpoint::MakeMessage<FStormSyncTransportWakeupRequest>();
+
+	TArray<FMessageAddress> Connections;
+	{
+		FScopeLock ConnectionTimeLock(&ConnectionLastActiveSection);
+		ConnectedAddresses.GetKeys(Connections);
+	}
+	
+	MessageEndpoint->Send(Message, Connections);
+}
+
 TSharedPtr<FMessageEndpoint, ESPMode::ThreadSafe> FStormSyncDiscoveryManager::GetMessageEndpoint() const
 {
 	return MessageEndpoint;
@@ -197,14 +218,49 @@ void FStormSyncDiscoveryManager::HandleConnectMessage(const FStormSyncTransportC
 
 	UE_LOG(LogStormSyncServer, VeryVerbose, TEXT("FStormSyncDiscoveryManager::HandleConnectMessage - SenderMessageAddress: %s, MessageEndpointAddress: %s (Same: %s)"), *SenderMessageAddress.ToString(), *MessageEndpointAddress.ToString(), SenderMessageAddress == MessageEndpointAddress ? TEXT("true") : TEXT("false"));
 
-	if (SenderMessageAddress != MessageEndpointAddress && !ConnectedAddresses.Contains(SenderMessageAddress))
+	if (SenderMessageAddress != MessageEndpointAddress)
 	{
-		// Add this message address to our list of connected message bus addresses
-		RegisterConnection(InMessage, SenderMessageAddress);
+		if (!IsConnectionRegistered(SenderMessageAddress))
+		{
+			// Add this message address to our list of connected message bus addresses
+			RegisterConnection(SenderMessageAddress, !InMessage.StormSyncServerAddressId.IsEmpty());
+		}
+
+		bool bNotifyConnectionInfo = false;
+		{
+			FScopeLock ConnectionTimeLock(&ConnectionLastActiveSection);
+			if (FStormSyncConnectedMessageBusAddress* FoundConnection = ConnectedAddresses.Find(SenderMessageAddress))
+			{
+				bNotifyConnectionInfo = !FoundConnection->bReceivedConnectionInfo;
+				FoundConnection->bReceivedConnectionInfo = true;
+			}
+		}
+		
+		if (bNotifyConnectionInfo)
+		{
+			// Notify editor of incoming connection
+			FStormSyncConnectedDevice ConnectedDevice;
+			ConnectedDevice.State = EStormSyncConnectedDeviceState::State_Active;
+			ConnectedDevice.MessageAddressId = SenderMessageAddress.ToString();
+			ConnectedDevice.bIsServerRunning = !InMessage.StormSyncServerAddressId.IsEmpty();
+			ConnectedDevice.StormSyncServerAddressId = InMessage.StormSyncServerAddressId;
+			ConnectedDevice.StormSyncClientAddressId = InMessage.StormSyncClientAddressId;
+			ConnectedDevice.HostName = InMessage.HostName;
+			ConnectedDevice.ProjectName = InMessage.ProjectName;
+			ConnectedDevice.ProjectDir = InMessage.ProjectDir;
+			ConnectedDevice.InstanceType = InMessage.InstanceType;
+
+			FStormSyncCoreDelegates::OnServiceDiscoveryConnection.Broadcast(SenderMessageAddress.ToString(), ConnectedDevice);
+		}
 	}
 }
 
-void FStormSyncDiscoveryManager::RegisterConnection(const FStormSyncTransportConnectMessage& InMessage, const FMessageAddress& InMessageAddress)
+void FStormSyncDiscoveryManager::HandleWakeupMessage(const FStormSyncTransportWakeupRequest& InMessage, const TSharedRef<IMessageContext, ESPMode::ThreadSafe>& MessageContext)
+{
+	FStormSyncCoreDelegates::OnServiceDiscoveryReceivedWakeup.Broadcast();
+}
+
+void FStormSyncDiscoveryManager::RegisterConnection(const FMessageAddress& InMessageAddress, bool bInIsServerRunning)
 {
 	{
 		FScopeLock ConnectionTimeLock(&ConnectionLastActiveSection);
@@ -212,7 +268,7 @@ void FStormSyncDiscoveryManager::RegisterConnection(const FStormSyncTransportCon
 
 		FStormSyncConnectedMessageBusAddress ConnectedAddress;
 		ConnectedAddress.bIsValid = true;
-		ConnectedAddress.bIsServerRunning = !InMessage.StormSyncServerAddressId.IsEmpty();
+		ConnectedAddress.bIsServerRunning = bInIsServerRunning;
 		ConnectedAddress.LastActivityTime = FPlatformTime::Seconds();
 
 		ConnectedAddresses.Add(InMessageAddress, ConnectedAddress);
@@ -231,20 +287,6 @@ void FStormSyncDiscoveryManager::RegisterConnection(const FStormSyncTransportCon
 		Message->StormSyncClientAddressId = IStormSyncTransportClientModule::Get().GetClientEndpointMessageAddressId();
 		MessageEndpoint->Send(Message.Release(), InMessageAddress);
 	}
-
-	// Notify editor of incoming connection
-	FStormSyncConnectedDevice ConnectedDevice;
-	ConnectedDevice.State = EStormSyncConnectedDeviceState::State_Active;
-	ConnectedDevice.MessageAddressId = InMessageAddress.ToString();
-	ConnectedDevice.bIsServerRunning = !InMessage.StormSyncServerAddressId.IsEmpty();
-	ConnectedDevice.StormSyncServerAddressId = InMessage.StormSyncServerAddressId;
-	ConnectedDevice.StormSyncClientAddressId = InMessage.StormSyncClientAddressId;
-	ConnectedDevice.HostName = InMessage.HostName;
-	ConnectedDevice.ProjectName = InMessage.ProjectName;
-	ConnectedDevice.ProjectDir = InMessage.ProjectDir;
-	ConnectedDevice.InstanceType = InMessage.InstanceType;
-
-	FStormSyncCoreDelegates::OnServiceDiscoveryConnection.Broadcast(InMessageAddress.ToString(), ConnectedDevice);
 }
 
 void FStormSyncDiscoveryManager::HandleHeartbeatMessage(const FStormSyncTransportHeartbeatMessage& InMessage, const TSharedRef<IMessageContext, ESPMode::ThreadSafe>& MessageContext)
@@ -253,28 +295,14 @@ void FStormSyncDiscoveryManager::HandleHeartbeatMessage(const FStormSyncTranspor
 
 	UE_LOG(LogStormSyncServer, VeryVerbose, TEXT("FStormSyncDiscoveryManager::HandleHeartbeatMessage - Received Heartbeat Message from %s"), *MessageAddress.ToString());
 	
+	// Handle revive connection in case we receive heartbeats again, most likely meaning
+	// connection was marked here as inactive due to inactive timeout and cleaned up.
+	//
+	// This is likely caused by a debug session with breakpoints that are longer than the configured inactive timeout
+	
+	if (!IsConnectionRegistered(MessageAddress))
 	{
-		// Handle revive connection in case we receive heartbeats again, most likely meaning
-		// connection was marked here as inactive due to inactive timeout and cleaned up.
-		//
-		// This is likely caused by a debug session with breakpoints that are longer than the configured inactive timeout
-
-		bool bContainsAddress;
-		{
-			FScopeLock ConnectionTimeLock(&ConnectionLastActiveSection);
-			bContainsAddress = ConnectedAddresses.Contains(MessageAddress);
-		}
-
-		if (!bContainsAddress)
-		{
-			TUniquePtr<FStormSyncTransportConnectMessage> ReviveMessage(FMessageEndpoint::MakeMessage<FStormSyncTransportConnectMessage>());
-			if (ReviveMessage.IsValid())
-			{
-				ReviveMessage->StormSyncServerAddressId = IStormSyncTransportServerModule::Get().GetServerEndpointMessageAddressId();
-				ReviveMessage->StormSyncClientAddressId = IStormSyncTransportClientModule::Get().GetClientEndpointMessageAddressId();
-				RegisterConnection(*ReviveMessage.Release(), MessageAddress);
-			}
-		}
+		RegisterConnection(MessageAddress, InMessage.bIsServerRunning);
 	}
 
 	UpdateConnectionLastActive(MessageAddress);
@@ -290,32 +318,23 @@ void FStormSyncDiscoveryManager::HandleHeartbeatMessage(const FStormSyncTranspor
 void FStormSyncDiscoveryManager::UpdateConnectionLastActive(const FMessageAddress& InAddress)
 {
 	FScopeLock ConnectionTimeLock(&ConnectionLastActiveSection);
-
-	if (ConnectedAddresses.Contains(InAddress))
+	if (FStormSyncConnectedMessageBusAddress* ConnectedMessageBusAddress = ConnectedAddresses.Find(InAddress))
 	{
-		FStormSyncConnectedMessageBusAddress& ConnectedMessageBusAddress = ConnectedAddresses.FindChecked(InAddress);
-		ConnectedMessageBusAddress.LastActivityTime = FPlatformTime::Seconds();
+		ConnectedMessageBusAddress->LastActivityTime = FPlatformTime::Seconds();
 	}
 }
 
 bool FStormSyncDiscoveryManager::UpdateServerStatus(const FMessageAddress& InAddress, const bool bIsServerRunning)
 {
 	FScopeLock ConnectionTimeLock(&ConnectionLastActiveSection);
-
-	bool bServerStatusChanged = false;
-
-	if (ConnectedAddresses.Contains(InAddress))
+	if (FStormSyncConnectedMessageBusAddress* ConnectedMessageBusAddress = ConnectedAddresses.Find(InAddress))
 	{
-		FStormSyncConnectedMessageBusAddress& ConnectedMessageBusAddress = ConnectedAddresses.FindChecked(InAddress);
-
 		// Server status changed for this recipient
-		if (ConnectedMessageBusAddress.bIsServerRunning != bIsServerRunning)
+		if (ConnectedMessageBusAddress->bIsServerRunning != bIsServerRunning)
 		{
-			bServerStatusChanged = true;
+			ConnectedMessageBusAddress->bIsServerRunning = bIsServerRunning;
+			return true;
 		}
-
-		ConnectedMessageBusAddress.bIsServerRunning = bIsServerRunning;
 	}
-
-	return bServerStatusChanged;
+	return false;
 }
diff --git a/Engine/Plugins/Experimental/StormSync/Source/StormSyncTransportServer/Private/ServiceDiscovery/StormSyncDiscoveryManager.h b/Engine/Plugins/Experimental/StormSync/Source/StormSyncTransportServer/Private/ServiceDiscovery/StormSyncDiscoveryManager.h
--- a/Engine/Plugins/Experimental/StormSync/Source/StormSyncTransportServer/Private/ServiceDiscovery/StormSyncDiscoveryManager.h
+++ b/Engine/Plugins/Experimental/StormSync/Source/StormSyncTransportServer/Private/ServiceDiscovery/StormSyncDiscoveryManager.h
@@ -23,6 +23,9 @@ struct FStormSyncConnectedMessageBusAddress
 	/** Represents remote server local endpoint status. Whether it is currently running and able to receive incoming connections */
 	bool bIsServerRunning = false;
 
+	/** Indicate if a connection info has been received for this connection. */
+	bool bReceivedConnectionInfo = false;
+
 	/** Default constructor */
 	FStormSyncConnectedMessageBusAddress() = default;
 };
@@ -81,6 +84,8 @@ public:
 	/** Broadcasts / Publish a FStormSyncTransportConnectMessage over the network */
 	void PublishConnectMessage();
 
+	void SendWakeUp();
+	
 	/** Returns the underlying endpoint so that outside code could send message directly to it (w/o relying on publish) */
 	TSharedPtr<FMessageEndpoint, ESPMode::ThreadSafe> GetMessageEndpoint() const;
 
@@ -128,13 +133,21 @@ private:
 	/** Callback handler to receive FStormSyncTransportConnectMessage messages */
 	void HandleConnectMessage(const FStormSyncTransportConnectMessage& InMessage, const TSharedRef<IMessageContext, ESPMode::ThreadSafe>& MessageContext);
 
+	void HandleWakeupMessage(const FStormSyncTransportWakeupRequest& InMessage, const TSharedRef<IMessageContext, ESPMode::ThreadSafe>& MessageContext);
+	
+	bool IsConnectionRegistered(const FMessageAddress& InMessageAddress)
+	{
+		FScopeLock ConnectionTimeLock(&ConnectionLastActiveSection);
+		return ConnectedAddresses.Contains(InMessageAddress);
+	}
+	
 	/**
 	 * Handles an incoming connect message by registering a new connection internally (adding to ConnectedAddresses), starting heartbeats for this recipient
 	 * and sending back a connect message to original sender.
 	 *
 	 * You must ensure connection is not already registered before calling this method.
 	 */
-	void RegisterConnection(const FStormSyncTransportConnectMessage& InMessage, const FMessageAddress& InMessageAddress);
+	void RegisterConnection(const FMessageAddress& InMessageAddress, bool bInIsServerRunning);
 
 	/** Callback handler to receive FStormSyncTransportHeartbeatMessage messages */
 	void HandleHeartbeatMessage(const FStormSyncTransportHeartbeatMessage& InMessage, const TSharedRef<IMessageContext, ESPMode::ThreadSafe>& MessageContext);
diff --git a/Engine/Plugins/Experimental/StormSync/Source/StormSyncTransportServer/Private/StormSyncTransportServerModule.cpp b/Engine/Plugins/Experimental/StormSync/Source/StormSyncTransportServer/Private/StormSyncTransportServerModule.cpp
--- a/Engine/Plugins/Experimental/StormSync/Source/StormSyncTransportServer/Private/StormSyncTransportServerModule.cpp
+++ b/Engine/Plugins/Experimental/StormSync/Source/StormSyncTransportServer/Private/StormSyncTransportServerModule.cpp
@@ -200,7 +200,6 @@ void FStormSyncTransportServerModule::RegisterConsoleCommands()
 	ConsoleCommands.Add(IConsoleManager::Get().RegisterConsoleCommand(
 		TEXT("StormSync.Server.Debug"),
 		TEXT("Prints out server address enpoint id to the console"),
-		// FConsoleCommandWithArgsDelegate::CreateRaw(this, &FStormSyncTransportClientModule::ExecuteSyncPak),
 		FConsoleCommandWithArgsDelegate::CreateLambda([this](const TArray<FString>& Args)
 		{
 			const FString AddressId = GetServerEndpointMessageAddressId();
@@ -208,6 +207,19 @@ void FStormSyncTransportServerModule::RegisterConsoleCommands()
 		}),
 		ECVF_Default
 	));
+
+	ConsoleCommands.Add(IConsoleManager::Get().RegisterConsoleCommand(
+		TEXT("StormSync.Discovery.Wakeup"),
+		TEXT("Send a wakeup request through discovery manager."),
+		FConsoleCommandWithArgsDelegate::CreateLambda([this](const TArray<FString>& Args)
+		{
+			if (DiscoveryManager)
+			{
+				DiscoveryManager->SendWakeUp();
+			}
+		}),
+		ECVF_Default
+	));
 }
 
 void FStormSyncTransportServerModule::UnregisterConsoleCommands()
