[StormSync] Improvements: FStormSyncStatusMessageService::HandleStatusRequest implements the task in async thread to no block the main thread.

This operation can be expensive and lock the main thread for a significant time and cause render nodes to time out and disconnect. This is also what cause the editor to lockup when selecting pages in rundown.

diff --git a/Engine/Plugins/Experimental/StormSync/Source/StormSyncTransportClient/Private/Services/StormSyncStatusMessageService.cpp b/Engine/Plugins/Experimental/StormSync/Source/StormSyncTransportClient/Private/Services/StormSyncStatusMessageService.cpp
--- a/Engine/Plugins/Experimental/StormSync/Source/StormSyncTransportClient/Private/Services/StormSyncStatusMessageService.cpp
+++ b/Engine/Plugins/Experimental/StormSync/Source/StormSyncTransportClient/Private/Services/StormSyncStatusMessageService.cpp
@@ -106,28 +106,40 @@ void FStormSyncStatusMessageService::HandleStatusPing(const FStormSyncTransportS
 
 void FStormSyncStatusMessageService::HandleStatusRequest(const FStormSyncTransportStatusRequest& InMessage, const TSharedRef<IMessageContext, ESPMode::ThreadSafe>& MessageContext)
 {
-	UE_LOG(LogStormSyncClient, Display, TEXT("FStormSyncStatusMessageService::HandleStatusRequest - Message: %s, Sender: %s"), *InMessage.ToString(), *MessageContext->GetSender().ToString());
-	const TSharedPtr<FMessageEndpoint, ESPMode::ThreadSafe> MessageEndpoint = GetMessageEndpoint();
-	if (!MessageEndpoint.IsValid())
-	{
-		UE_LOG(LogStormSyncClient, Error, TEXT("FStormSyncStatusMessageService::HandleStatusPing - Unable to get client endpoint"));
-		return;
-	}
+	FMessageAddress Sender = MessageContext->GetSender();
 
-	FStormSyncTransportStatusResponse* Message = FMessageEndpoint::MakeMessage<FStormSyncTransportStatusResponse>(InMessage.MessageId);
+	UE_LOG(LogStormSyncClient, Display,
+		TEXT("FStormSyncStatusMessageService::HandleStatusRequest - Message: %s, Sender: %s"), *InMessage.ToString(), *Sender.ToString());
 
-	// Compute the diff
-	Message->Modifiers = FStormSyncCoreUtils::GetSyncFileModifiers(InMessage.PackageNames, InMessage.Dependencies);
-	Message->bNeedsSynchronization = !Message->Modifiers.IsEmpty();
+	const TWeakPtr<FMessageEndpoint> MessageEndpointWeak = GetMessageEndpoint();
 
-	UE_LOG(LogStormSyncClient, Display, TEXT("GetSyncFileModifiers ... - Modifiers: %d"), Message->Modifiers.Num());
-	for (const FStormSyncFileModifierInfo& Modifier : Message->Modifiers)
+	Async(EAsyncExecution::Thread, [Sender, MessageEndpointWeak, MessageId = InMessage.MessageId, PackageNames = InMessage.PackageNames, Dependencies = InMessage.Dependencies]()
 	{
-		UE_LOG(LogStormSyncClient, Display, TEXT("\t Modifier: %s"), *Modifier.ToString());
-	}
+		TUniquePtr<FStormSyncTransportStatusResponse> Message(FMessageEndpoint::MakeMessage<FStormSyncTransportStatusResponse>(MessageId));
 
-	// Send back result
-	MessageEndpoint->Send(Message, MessageContext->GetSender());
+		UE_LOG(LogStormSyncClient, Display, TEXT("GetSyncFileModifiers for Packages %s"),
+			*FString::JoinBy(PackageNames, TEXT(" ,"), [](const FName& InName){return InName.ToString();}));
+	
+		// Compute the diff
+		Message->Modifiers = FStormSyncCoreUtils::GetSyncFileModifiers(PackageNames, Dependencies);
+		Message->bNeedsSynchronization = !Message->Modifiers.IsEmpty();
+
+		UE_LOG(LogStormSyncClient, Display, TEXT("GetSyncFileModifiers ... - Modifiers: %d"), Message->Modifiers.Num());
+		for (const FStormSyncFileModifierInfo& Modifier : Message->Modifiers)
+		{
+			UE_LOG(LogStormSyncClient, Display, TEXT("\t Modifier: %s"), *Modifier.ToString());
+		}
+
+		const TSharedPtr<FMessageEndpoint, ESPMode::ThreadSafe> MessageEndpoint = MessageEndpointWeak.Pin();
+		if (!MessageEndpoint.IsValid())
+		{
+			UE_LOG(LogStormSyncClient, Error, TEXT("FStormSyncStatusMessageService::HandleStatusRequest - Unable to get client endpoint"));
+			return;
+		}
+		
+		// Send back result
+		MessageEndpoint->Send(Message.Release(), Sender);
+	});
 }
 
 void FStormSyncStatusMessageService::HandleStatusResponse(const FStormSyncTransportStatusResponse& InMessage, const TSharedRef<IMessageContext, ESPMode::ThreadSafe>& MessageContext)
@@ -195,7 +207,7 @@ TSharedPtr<FMessageEndpoint, ESPMode::ThreadSafe> FStormSyncStatusMessageService
 {
 	if (const TSharedPtr<IStormSyncTransportLocalEndpoint> LocalEndpoint = GetLocalEndpoint())
 	{
-		return LocalEndpoint->GetMessageEndpoint();	
+		return LocalEndpoint->GetMessageEndpoint();
 	}
 	
 	return nullptr;
