[StormSync] Add ability to override tcp server port (only) as part of the command line while preserving the host name from udp settings.

diff --git a/Engine/Plugins/Experimental/StormSync/Source/StormSyncTransportCore/Private/StormSyncTransportSettings.cpp b/Engine/Plugins/Experimental/StormSync/Source/StormSyncTransportCore/Private/StormSyncTransportSettings.cpp
--- a/Engine/Plugins/Experimental/StormSync/Source/StormSyncTransportCore/Private/StormSyncTransportSettings.cpp
+++ b/Engine/Plugins/Experimental/StormSync/Source/StormSyncTransportCore/Private/StormSyncTransportSettings.cpp
@@ -36,32 +36,41 @@ void UStormSyncTransportSettings::OpenEditorSettingsWindow() const
 
 FString UStormSyncTransportSettings::GetServerEndpoint() const
 {
-	FString ServerEndpoint = FString::Printf(TEXT("%s:%d"), *TcpServerAddress, TcpServerPort);
-
-	// Take the command line overrides if they are used and valid (eg. `-StormSyncServerEndpoint=`)
-	FString CommandLineOverride;
-	if (UE::StormSync::Transport::Private::GetServerEndpointParam(CommandLineOverride))
-	{
-		return CommandLineOverride;
-	}
-
+	bool bReportUdpSettingError = false;
+	FString ServerAddress = TcpServerAddress;
 	if (!bOverrideServerAddress)
 	{
 		// Server override is disabled in project settings (the default behavior)
 		// Use the same value as defined in UDPMessagingSettings > Unicast Endpoint
 		FString UnicastPort;
 		FString UnicastHostname;
-		if (!GetUdpMessagingUnicastEndpoint(UnicastHostname, UnicastPort))
+		if (GetUdpMessagingUnicastEndpoint(UnicastHostname, UnicastPort))
+		{
+			ServerAddress = UnicastHostname;
+		}
+		else
 		{
 			// Either unable to get the setting or malformed
-			UE_LOG(LogStormSyncTransportCore, Display, TEXT("UStormSyncTransportSettings::GetServerEndpoint - Was unable to get UnicastEndpoint from UDPMessagingSettings"));
-			return ServerEndpoint;
+			bReportUdpSettingError = true;
 		}
-
-		return FString::Printf(TEXT("%s:%d"), *UnicastHostname, TcpServerPort);
 	}
 	
-	return ServerEndpoint;
+	// Take the command line overrides if they are used and valid (eg. `-StormSyncServerEndpoint=`)
+	FString CommandLineOverride;
+	if (UE::StormSync::Transport::Private::GetServerEndpointParam(CommandLineOverride, ServerAddress))
+	{
+		return CommandLineOverride;
+	}
+
+	// Report the previous error if we end up using the address from the settings.
+	if (bReportUdpSettingError)
+	{
+		UE_LOG(LogStormSyncTransportCore, Display,
+			TEXT("UStormSyncTransportSettings::GetServerEndpoint - Was unable to get UnicastEndpoint from UDPMessagingSettings, using \"%s\" instead.")
+			, *ServerAddress);
+	}
+
+	return FString::Printf(TEXT("%s:%d"), *ServerAddress, TcpServerPort);
 }
 
 FString UStormSyncTransportSettings::GetTcpServerAddress() const
@@ -71,6 +80,7 @@ FString UStormSyncTransportSettings::GetTcpServerAddress() const
 
 uint16 UStormSyncTransportSettings::GetTcpServerPort() const
 {
+	
 	return TcpServerPort;
 }
 
diff --git a/Engine/Plugins/Experimental/StormSync/Source/StormSyncTransportCore/Private/Utils/StormSyncTransportCommandUtils.cpp b/Engine/Plugins/Experimental/StormSync/Source/StormSyncTransportCore/Private/Utils/StormSyncTransportCommandUtils.cpp
--- a/Engine/Plugins/Experimental/StormSync/Source/StormSyncTransportCore/Private/Utils/StormSyncTransportCommandUtils.cpp
+++ b/Engine/Plugins/Experimental/StormSync/Source/StormSyncTransportCore/Private/Utils/StormSyncTransportCommandUtils.cpp
@@ -6,7 +6,7 @@
 #include "Misc/Parse.h"
 #include "StormSyncTransportSettings.h"
 
-bool UE::StormSync::Transport::Private::GetServerEndpointParam(FString& OutEndpointValue)
+bool UE::StormSync::Transport::Private::GetServerEndpointParam(FString& OutEndpointValue, const FString& InDefaultHostName)
 {
 	FString ServerEndpoint;
 	if (!FParse::Value(FCommandLine::Get(), ServerEndpointFlag, ServerEndpoint))
@@ -27,7 +27,7 @@ bool UE::StormSync::Transport::Private::GetServerEndpointParam(FString& OutEndpo
 		uint16 ServerPort;
 		if (FParse::Value(FCommandLine::Get(), ServerEndpointFlag, ServerPort))
 		{
-			const FString DefaultHostname = GetDefault<UStormSyncTransportSettings>()->GetTcpServerAddress();
+			const FString DefaultHostname = !InDefaultHostName.IsEmpty() ? InDefaultHostName : GetDefault<UStormSyncTransportSettings>()->GetTcpServerAddress();
 			OutEndpointValue = FString::Printf(TEXT("%s:%d"), *DefaultHostname, ServerPort);
 			return true;
 		}
diff --git a/Engine/Plugins/Experimental/StormSync/Source/StormSyncTransportCore/Public/Utils/StormSyncTransportCommandUtils.h b/Engine/Plugins/Experimental/StormSync/Source/StormSyncTransportCore/Public/Utils/StormSyncTransportCommandUtils.h
--- a/Engine/Plugins/Experimental/StormSync/Source/StormSyncTransportCore/Public/Utils/StormSyncTransportCommandUtils.h
+++ b/Engine/Plugins/Experimental/StormSync/Source/StormSyncTransportCore/Public/Utils/StormSyncTransportCommandUtils.h
@@ -18,7 +18,7 @@ namespace UE::StormSync::Transport::Private
 	 *
 	 * OutEndpointValue will be updated only if the parsing was successful.
 	 */
-	STORMSYNCTRANSPORTCORE_API bool GetServerEndpointParam(FString& OutEndpointValue);
+	STORMSYNCTRANSPORTCORE_API bool GetServerEndpointParam(FString& OutEndpointValue, const FString& InDefaultHostName = FString());
 	
 	/** Returns whether command line has either -StormSyncServerAutoStart or a valid -StormSyncServerEndpoint=<Port> */
 	STORMSYNCTRANSPORTCORE_API bool IsServerAutoStartDisabled();
