[StormSync] Support Multiple Storm Sync TCP servers on the same host.

Some configurations require running multiple unreal processes, either editors or games, with storm sync tcp servers on the same host.

In order to make this easier, the tcp servers will try to find an available port automatically upon finding a collision.

Implementation notes:
- The messages for the pull, push and compare services are in the client module, but need the server tcp address. But also, the discovery service being on the server module (could be it's own module), requiring the client endpoint Id. These dependencies need to be untangled. So, I just added some delegates in the transport core to solve this. This way, client module no longer need a dependency on the server module.

diff --git a/Engine/Plugins/Experimental/StormSync/Source/StormSyncTransportClient/Private/Services/StormSyncPullMessageService.cpp b/Engine/Plugins/Experimental/StormSync/Source/StormSyncTransportClient/Private/Services/StormSyncPullMessageService.cpp
--- a/Engine/Plugins/Experimental/StormSync/Source/StormSyncTransportClient/Private/Services/StormSyncPullMessageService.cpp
+++ b/Engine/Plugins/Experimental/StormSync/Source/StormSyncTransportClient/Private/Services/StormSyncPullMessageService.cpp
@@ -99,7 +99,7 @@ void FStormSyncPullMessageService::RequestPullPackages(const FMessageAddress& In
 
 		// Fill in Pull Request Message with network info for this device
 		PullRequestMessage->HostName = FStormSyncTransportNetworkUtils::GetServerName();
-		PullRequestMessage->HostAddress = FStormSyncTransportNetworkUtils::GetTcpEndpointAddress();
+		PullRequestMessage->HostAddress = FStormSyncTransportNetworkUtils::GetCurrentTcpServerEndpointAddress();
 		PullRequestMessage->HostAdapterAddresses = FStormSyncTransportNetworkUtils::GetLocalAdapterAddresses();
 
 		UE_LOG(LogStormSyncClient, Display, TEXT("FStormSyncPullMessageService::RequestPullPackages - FileDependencies: %d"), FileDependencies.Num());
@@ -343,7 +343,7 @@ FStormSyncTransportPullResponse FStormSyncPullMessageService::CreatePullResponse
 	
 	// Fill in payload with network info for this device
 	Payload.HostName = FStormSyncTransportNetworkUtils::GetServerName();
-	Payload.HostAddress = FStormSyncTransportNetworkUtils::GetTcpEndpointAddress();
+	Payload.HostAddress = FStormSyncTransportNetworkUtils::GetCurrentTcpServerEndpointAddress();
 	Payload.HostAdapterAddresses = FStormSyncTransportNetworkUtils::GetLocalAdapterAddresses();
 	
 	return Payload;
diff --git a/Engine/Plugins/Experimental/StormSync/Source/StormSyncTransportClient/Private/Services/StormSyncPushMessageService.cpp b/Engine/Plugins/Experimental/StormSync/Source/StormSyncTransportClient/Private/Services/StormSyncPushMessageService.cpp
--- a/Engine/Plugins/Experimental/StormSync/Source/StormSyncTransportClient/Private/Services/StormSyncPushMessageService.cpp
+++ b/Engine/Plugins/Experimental/StormSync/Source/StormSyncTransportClient/Private/Services/StormSyncPushMessageService.cpp
@@ -150,7 +150,7 @@ void FStormSyncPushMessageService::HandlePushRequestMessage(const FStormSyncTran
 		Message->Modifiers = FileModifiers;
 
 		Message->HostName = FStormSyncTransportNetworkUtils::GetServerName();
-		Message->HostAddress = FStormSyncTransportNetworkUtils::GetTcpEndpointAddress();
+		Message->HostAddress = FStormSyncTransportNetworkUtils::GetCurrentTcpServerEndpointAddress();
 		Message->HostAdapterAddresses = FStormSyncTransportNetworkUtils::GetLocalAdapterAddresses();
 
 		UE_LOG(LogStormSyncClient, Display, TEXT("FStormSyncPushMessageService::HandlePushRequestMessage - Sending back to %s sync response: %s"), *RemoteAddress.ToString(), *Message->ToString());
@@ -273,7 +273,7 @@ FStormSyncTransportPushResponse FStormSyncPushMessageService::CreatePushResponse
 	
 	// Fill in payload with network info for this device
 	Payload.HostName = FStormSyncTransportNetworkUtils::GetServerName();
-	Payload.HostAddress = FStormSyncTransportNetworkUtils::GetTcpEndpointAddress();
+	Payload.HostAddress = FStormSyncTransportNetworkUtils::GetCurrentTcpServerEndpointAddress();
 	Payload.HostAdapterAddresses = FStormSyncTransportNetworkUtils::GetLocalAdapterAddresses();
 	
 	return Payload;
diff --git a/Engine/Plugins/Experimental/StormSync/Source/StormSyncTransportClient/Private/StormSyncTransportClientModule.cpp b/Engine/Plugins/Experimental/StormSync/Source/StormSyncTransportClient/Private/StormSyncTransportClientModule.cpp
--- a/Engine/Plugins/Experimental/StormSync/Source/StormSyncTransportClient/Private/StormSyncTransportClientModule.cpp
+++ b/Engine/Plugins/Experimental/StormSync/Source/StormSyncTransportClient/Private/StormSyncTransportClientModule.cpp
@@ -4,6 +4,7 @@
 
 #include "StormSyncCommandLineUtils.h"
 #include "StormSyncCoreUtils.h"
+#include "IStormSyncTransportCoreModule.h"
 #include "IStormSyncTransportLocalEndpoint.h"
 #include "MessageEndpoint.h"
 #include "Misc/CoreDelegates.h"
@@ -28,12 +29,19 @@ void FStormSyncTransportClientModule::StartupModule()
 	FCoreDelegates::OnPostEngineInit.AddRaw(this, &FStormSyncTransportClientModule::OnPostEngineInit);
 
 	FStormSyncCoreDelegates::OnServiceDiscoveryReceivedWakeup.AddRaw(this, &FStormSyncTransportClientModule::OnWakeup);
+	
+	IStormSyncTransportCoreModule::Get().OnGetClientEndpointMessageAddress().BindRaw(this, &FStormSyncTransportClientModule::GetClientEndpointMessageAddressId);
 }
 
 void FStormSyncTransportClientModule::ShutdownModule()
 {
 	FCoreDelegates::OnPostEngineInit.RemoveAll(this);
 	FStormSyncCoreDelegates::OnServiceDiscoveryReceivedWakeup.RemoveAll(this);
+
+	if (IStormSyncTransportCoreModule::IsAvailable())
+	{
+		IStormSyncTransportCoreModule::Get().OnGetClientEndpointMessageAddress().Unbind();
+	}
 	
 	if (ClientEndpoint.IsValid())
 	{
diff --git a/Engine/Plugins/Experimental/StormSync/Source/StormSyncTransportCore/Internal/IStormSyncTransportCoreModule.h b/Engine/Plugins/Experimental/StormSync/Source/StormSyncTransportCore/Internal/IStormSyncTransportCoreModule.h
new file mode 100644
--- /dev/null
+++ b/Engine/Plugins/Experimental/StormSync/Source/StormSyncTransportCore/Internal/IStormSyncTransportCoreModule.h
@@ -0,0 +1,45 @@
+// Copyright Epic Games, Inc. All Rights Reserved.
+
+#pragma once
+
+#include "Modules/ModuleInterface.h"
+#include "Modules/ModuleManager.h"
+
+class IStormSyncTransportCoreModule : public IModuleInterface
+{
+public:
+	/**
+	 * Singleton-like access to this module's interface.  This is just for convenience!
+	 * Beware of calling this during the shutdown phase, though.  Your module might have been unloaded already.
+	 *
+	 * @return Returns singleton instance, loading the module on demand if needed
+	 */
+	static IStormSyncTransportCoreModule& Get()
+	{
+		static const FName ModuleName = "StormSyncTransportCore";
+		return FModuleManager::LoadModuleChecked<IStormSyncTransportCoreModule>(ModuleName);
+	}
+
+	/**
+	 * Checks to see if this module is loaded and ready.  It is only valid to call Get() during shutdown if IsAvailable() returns true.
+	 *
+	 * @return True if the module is loaded and ready to use
+	 */
+	static bool IsAvailable()
+	{
+		static const FName ModuleName = "StormSyncTransportCore";
+		return FModuleManager::Get().IsModuleLoaded(ModuleName);
+	}
+
+	/** Definition of a delegate to query endpoint configuration. */
+	DECLARE_DELEGATE_RetVal(FString, FOnGetEndpointConfig);
+
+	/** Delegate for querying the currently bound tcp server address. */
+	virtual FOnGetEndpointConfig& OnGetCurrentTcpServerEndpointAddress() = 0;
+
+	/** Delegate for querying the message bus server endpoint address. */
+	virtual FOnGetEndpointConfig& OnGetServerEndpointMessageAddress() = 0;
+
+	/** Delegate for querying the message bus client endpoint address. */
+	virtual FOnGetEndpointConfig& OnGetClientEndpointMessageAddress() = 0;
+};
diff --git a/Engine/Plugins/Experimental/StormSync/Source/StormSyncTransportCore/Private/StormSyncTransportCoreModule.cpp b/Engine/Plugins/Experimental/StormSync/Source/StormSyncTransportCore/Private/StormSyncTransportCoreModule.cpp
--- a/Engine/Plugins/Experimental/StormSync/Source/StormSyncTransportCore/Private/StormSyncTransportCoreModule.cpp
+++ b/Engine/Plugins/Experimental/StormSync/Source/StormSyncTransportCore/Private/StormSyncTransportCoreModule.cpp
@@ -1,5 +1,5 @@
 // Copyright Epic Games, Inc. All Rights Reserved.
 
-#include "Modules/ModuleManager.h"
+#include "StormSyncTransportCoreModule.h"
 
-IMPLEMENT_MODULE(FDefaultModuleImpl, StormSyncTransportCore)
+IMPLEMENT_MODULE(FStormSyncTransportCoreModule, StormSyncTransportCore)
\ No newline at end of file
diff --git a/Engine/Plugins/Experimental/StormSync/Source/StormSyncTransportCore/Private/StormSyncTransportCoreModule.h b/Engine/Plugins/Experimental/StormSync/Source/StormSyncTransportCore/Private/StormSyncTransportCoreModule.h
new file mode 100644
--- /dev/null
+++ b/Engine/Plugins/Experimental/StormSync/Source/StormSyncTransportCore/Private/StormSyncTransportCoreModule.h
@@ -0,0 +1,25 @@
+// Copyright Epic Games, Inc. All Rights Reserved.
+
+#pragma once
+
+#include "IStormSyncTransportCoreModule.h"
+
+class FStormSyncDiscoveryManager;
+
+/**
+ * Implements the StormSyncTransportCore module.
+ */
+class FStormSyncTransportCoreModule : public IStormSyncTransportCoreModule
+{
+public:
+	//~ Begin IStormSyncTransportCoreModule
+	virtual FOnGetEndpointConfig& OnGetCurrentTcpServerEndpointAddress() override { return OnGetCurrentTcpServerAddressDelegate; }
+	virtual FOnGetEndpointConfig& OnGetServerEndpointMessageAddress() override { return OnGetServerEndpointMessageAddressDelegate; }
+	virtual FOnGetEndpointConfig& OnGetClientEndpointMessageAddress() override { return OnGetClientEndpointMessageAddressDelegate; }
+	//~ End IStormSyncTransportCoreModule
+
+private:
+	FOnGetEndpointConfig OnGetCurrentTcpServerAddressDelegate;
+	FOnGetEndpointConfig OnGetServerEndpointMessageAddressDelegate;
+	FOnGetEndpointConfig OnGetClientEndpointMessageAddressDelegate;
+};
diff --git a/Engine/Plugins/Experimental/StormSync/Source/StormSyncTransportCore/Private/StormSyncTransportNetworkUtils.cpp b/Engine/Plugins/Experimental/StormSync/Source/StormSyncTransportCore/Private/StormSyncTransportNetworkUtils.cpp
--- a/Engine/Plugins/Experimental/StormSync/Source/StormSyncTransportCore/Private/StormSyncTransportNetworkUtils.cpp
+++ b/Engine/Plugins/Experimental/StormSync/Source/StormSyncTransportCore/Private/StormSyncTransportNetworkUtils.cpp
@@ -3,9 +3,18 @@
 #include "StormSyncTransportNetworkUtils.h"
 
 #include "Algo/Transform.h"
+#include "IStormSyncTransportCoreModule.h"
+#include "Interfaces/IPv4/IPv4Endpoint.h"
 #include "StormSyncTransportCoreLog.h"
 #include "StormSyncTransportSettings.h"
-#include "Interfaces/IPv4/IPv4Endpoint.h"
+
+namespace UE::StormSyncTransportCore::Utils::Private
+{
+	FString ExecuteIfBound(IStormSyncTransportCoreModule::FOnGetEndpointConfig& InDelegate, const FString& InDefaultValue)
+	{
+		return InDelegate.IsBound() ? InDelegate.Execute() : InDefaultValue;
+	}
+}
 
 FString FStormSyncTransportNetworkUtils::GetServerName()
 {
@@ -59,3 +68,21 @@ TArray<FString> FStormSyncTransportNetworkUtils::GetLocalAdapterAddresses()
 	
 	return AdapterAddresses;
 }
+
+FString FStormSyncTransportNetworkUtils::GetCurrentTcpServerEndpointAddress()
+{
+	using namespace UE::StormSyncTransportCore::Utils::Private;
+	return ExecuteIfBound(IStormSyncTransportCoreModule::Get().OnGetCurrentTcpServerEndpointAddress(), FString());
+}
+
+FString FStormSyncTransportNetworkUtils::GetServerEndpointMessageAddress()
+{
+	using namespace UE::StormSyncTransportCore::Utils::Private;
+	return ExecuteIfBound(IStormSyncTransportCoreModule::Get().OnGetServerEndpointMessageAddress(), FString());
+}
+
+FString FStormSyncTransportNetworkUtils::GetClientEndpointMessageAddress()
+{
+	using namespace UE::StormSyncTransportCore::Utils::Private;
+	return ExecuteIfBound(IStormSyncTransportCoreModule::Get().OnGetClientEndpointMessageAddress(), FString());
+}
diff --git a/Engine/Plugins/Experimental/StormSync/Source/StormSyncTransportCore/Public/StormSyncTransportNetworkUtils.h b/Engine/Plugins/Experimental/StormSync/Source/StormSyncTransportCore/Public/StormSyncTransportNetworkUtils.h
--- a/Engine/Plugins/Experimental/StormSync/Source/StormSyncTransportCore/Public/StormSyncTransportNetworkUtils.h
+++ b/Engine/Plugins/Experimental/StormSync/Source/StormSyncTransportCore/Public/StormSyncTransportNetworkUtils.h
@@ -26,4 +26,16 @@ public:
 	 * endpoint settings appended to the address list.
 	 */
 	static TArray<FString> GetLocalAdapterAddresses();
+
+	/**
+	 * Returns the current tcp server endpoint address (ip:port).
+	 * @remark This is potentially different than the configured endpoint address in case of port collisions, it may be on a different one.
+	 */
+	static FString GetCurrentTcpServerEndpointAddress();
+
+	/** Returns the current message bus server endpoint address. */
+	static FString GetServerEndpointMessageAddress();
+
+	/** Returns the current message bus client endpoint address. */
+	static FString GetClientEndpointMessageAddress();
 };
diff --git a/Engine/Plugins/Experimental/StormSync/Source/StormSyncTransportCore/Public/StormSyncTransportSettings.h b/Engine/Plugins/Experimental/StormSync/Source/StormSyncTransportCore/Public/StormSyncTransportSettings.h
--- a/Engine/Plugins/Experimental/StormSync/Source/StormSyncTransportCore/Public/StormSyncTransportSettings.h
+++ b/Engine/Plugins/Experimental/StormSync/Source/StormSyncTransportCore/Public/StormSyncTransportSettings.h
@@ -45,6 +45,11 @@ public:
 	/** Returns the port tcp server is configured to run on */
 	uint16 GetTcpServerPort() const;
 
+	/**
+	 * Returns th port offset range when trying to acquire an available socket for the Tcp server.
+	 */
+	uint16 GetTcpServerPortMaxOffset() const { return TcpServerPortMaxOffset; }
+
 	/** Returns inactive timeout period. The amount of time in seconds to consider a client to be inactive */
 	uint32 GetInactiveTimeoutSeconds() const;
 
@@ -109,8 +114,15 @@ private:
 	 * Can be specified on the command line with `-StormSyncServerEndpoint=`
 	 */
 	UPROPERTY(config, EditAnywhere, Category = "Server")
-	uint16 TcpServerPort = 40999;
+	uint16 TcpServerPort = 40990;
 
+	/**
+	 * Specifies the port offset range when trying to acquire an available socket for the Tcp server.
+	 * The range of ports will be [TcpServerPort, TcpServerPort + TcpServerPortMaxOffset - 1].
+	 */
+	UPROPERTY(config, EditAnywhere, Category = "Server")
+	uint16 TcpServerPortMaxOffset = 10;
+	
 	/**
 	 * User friendly name to describe this server on the network (defaults to hostname if left to empty)
 	 */
diff --git a/Engine/Plugins/Experimental/StormSync/Source/StormSyncTransportServer/Private/ServiceDiscovery/StormSyncDiscoveryManager.cpp b/Engine/Plugins/Experimental/StormSync/Source/StormSyncTransportServer/Private/ServiceDiscovery/StormSyncDiscoveryManager.cpp
--- a/Engine/Plugins/Experimental/StormSync/Source/StormSyncTransportServer/Private/ServiceDiscovery/StormSyncDiscoveryManager.cpp
+++ b/Engine/Plugins/Experimental/StormSync/Source/StormSyncTransportServer/Private/ServiceDiscovery/StormSyncDiscoveryManager.cpp
@@ -4,12 +4,12 @@
 
 #include "Async/TaskGraphInterfaces.h"
 #include "HAL/RunnableThread.h"
-#include "IStormSyncTransportClientModule.h"
 #include "IStormSyncTransportServerModule.h"
 #include "MessageEndpointBuilder.h"
 #include "StormSyncCoreDelegates.h"
 #include "StormSyncHeartbeatEmitter.h"
 #include "StormSyncTransportMessages.h"
+#include "StormSyncTransportNetworkUtils.h"
 #include "StormSyncTransportServerLog.h"
 
 FStormSyncDiscoveryManager::FStormSyncDiscoveryManager(double InHeartbeatTimeout, double InInactiveSourceTimeout, float InTickInterval, bool bInEnableDiscoveryPeriodicPublish)
@@ -159,8 +159,8 @@ void FStormSyncDiscoveryManager::PublishConnectMessage()
 	TUniquePtr<FStormSyncTransportConnectMessage> Message(FMessageEndpoint::MakeMessage<FStormSyncTransportConnectMessage>());
 	if (Message.IsValid())
 	{
-		Message->StormSyncServerAddressId = IStormSyncTransportServerModule::Get().GetServerEndpointMessageAddressId();
-		Message->StormSyncClientAddressId = IStormSyncTransportClientModule::Get().GetClientEndpointMessageAddressId();
+		Message->StormSyncServerAddressId = FStormSyncTransportNetworkUtils::GetServerEndpointMessageAddress();
+		Message->StormSyncClientAddressId = FStormSyncTransportNetworkUtils::GetClientEndpointMessageAddress();
 
 		LastPublishTime = FApp::GetCurrentTime();
 		MessageEndpoint->Publish(Message.Release());
@@ -283,8 +283,8 @@ void FStormSyncDiscoveryManager::RegisterConnection(const FMessageAddress& InMes
 	TUniquePtr<FStormSyncTransportConnectMessage> Message(FMessageEndpoint::MakeMessage<FStormSyncTransportConnectMessage>());
 	if (Message.IsValid())
 	{
-		Message->StormSyncServerAddressId = IStormSyncTransportServerModule::Get().GetServerEndpointMessageAddressId();
-		Message->StormSyncClientAddressId = IStormSyncTransportClientModule::Get().GetClientEndpointMessageAddressId();
+		Message->StormSyncServerAddressId = FStormSyncTransportNetworkUtils::GetServerEndpointMessageAddress();
+		Message->StormSyncClientAddressId = FStormSyncTransportNetworkUtils::GetClientEndpointMessageAddress();
 		MessageEndpoint->Send(Message.Release(), InMessageAddress);
 	}
 }
diff --git a/Engine/Plugins/Experimental/StormSync/Source/StormSyncTransportServer/Private/Socket/StormSyncTransportTcpServer.cpp b/Engine/Plugins/Experimental/StormSync/Source/StormSyncTransportServer/Private/Socket/StormSyncTransportTcpServer.cpp
--- a/Engine/Plugins/Experimental/StormSync/Source/StormSyncTransportServer/Private/Socket/StormSyncTransportTcpServer.cpp
+++ b/Engine/Plugins/Experimental/StormSync/Source/StormSyncTransportServer/Private/Socket/StormSyncTransportTcpServer.cpp
@@ -4,14 +4,13 @@
 
 #include "Common/TcpListener.h"
 #include "StormSyncTransportServerLog.h"
+#include "StormSyncTransportSettings.h"
 #include "StormSyncTransportTcpPackets.h"
 
 FStormSyncTransportTcpServer::FStormSyncTransportTcpServer(const FIPv4Address& InEndpointAddress, const uint16 InEndpointPort, const uint32 InInactiveTimeoutSeconds)
 	: DefaultInactiveTimeoutSeconds(InInactiveTimeoutSeconds)
-	, EndpointAddress(InEndpointAddress)
-	, EndpointPort(InEndpointPort)
+	, DesiredEndpoint(InEndpointAddress, InEndpointPort)
 {
-	Endpoint = MakeUnique<FIPv4Endpoint>(EndpointAddress, EndpointPort);
 	Thread = FRunnableThread::Create(this, TEXT("FStormSyncServerListener"), 128 * 1024, TPri_Normal);
 }
 
@@ -153,15 +152,40 @@ bool FStormSyncTransportTcpServer::Tick()
 bool FStormSyncTransportTcpServer::StartListening()
 {
 	check(!SocketListener);
+	
+	const uint16 MaxPortOffset = UStormSyncTransportSettings::Get().GetTcpServerPortMaxOffset();
+	
+	FSocket* RawSocket = nullptr;
 
-	SocketListener = MakeUnique<FTcpListener>(*Endpoint, FTimespan::FromSeconds(1), false);
-	if (SocketListener->IsActive())
+	for (uint16 PortOffset = 0; PortOffset < MaxPortOffset && !RawSocket; ++PortOffset)
 	{
-		SocketListener->OnConnectionAccepted().BindRaw(this, &FStormSyncTransportTcpServer::OnIncomingConnection);
-		UE_LOG(LogStormSyncServer, Display, TEXT("Started listening on %s:%d"), *SocketListener->GetLocalEndpoint().Address.ToString(), SocketListener->GetLocalEndpoint().Port);
-		return true;
+		FIPv4Endpoint ListenEndPoint(DesiredEndpoint.Address, DesiredEndpoint.Port + PortOffset);
+
+		RawSocket = FTcpSocketBuilder(TEXT("FStormSyncTransportTcpServer"))
+		.BoundToEndpoint(ListenEndPoint)
+		.Listening(8) // Max number of connections in queue before refusing them
+		.WithSendBufferSize(2 * 1024 * 1024).Build();
 	}
 
+	if (RawSocket)
+	{
+		// Wrap socket in a smart pointer so it is managed. It is not managed by FTcpListener.
+		Socket.Reset(RawSocket);
+
+		const FTimespan SleepTime =  FTimespan::FromSeconds(1);	// default argument.
+		constexpr bool bSocketReusable = false;
+
+		SocketListener = MakeUnique<FTcpListener>(*Socket, SleepTime, bSocketReusable);
+
+		if (SocketListener->IsActive())
+		{
+			SocketListener->OnConnectionAccepted().BindRaw(this, &FStormSyncTransportTcpServer::OnIncomingConnection);
+			UE_LOG(LogStormSyncServer, Display, TEXT("Started listening on %s"), *SocketListener->GetLocalEndpoint().ToString());
+			return true;
+		}
+	}
+
+	Socket.Reset();
 	UE_LOG(LogStormSyncServer, Error, TEXT("Could not create Tcp Listener!"));
 	return false;
 }
@@ -172,10 +196,12 @@ void FStormSyncTransportTcpServer::StopListening()
 
 	if (SocketListener.IsValid())
 	{
-		UE_LOG(LogStormSyncServer, Display, TEXT("No longer listening on %s:%d"), *SocketListener->GetLocalEndpoint().Address.ToString(), SocketListener->GetLocalEndpoint().Port);
+		UE_LOG(LogStormSyncServer, Display, TEXT("No longer listening on %s"), *SocketListener->GetLocalEndpoint().ToString());
 		SocketListener.Reset();
 	}
 
+	Socket.Reset();
+
 	for (const TTuple<FIPv4Endpoint, TUniquePtr<FClientConnection>>& Connection : ClientConnections) 
 	{
 		if (Connection.Value && Connection.Value->Socket.IsValid())
@@ -191,7 +217,12 @@ void FStormSyncTransportTcpServer::StopListening()
 
 FString FStormSyncTransportTcpServer::GetEndpointAddress() const
 {
-	return Endpoint->ToString();
+	if (SocketListener && SocketListener->IsActive())
+	{
+		return SocketListener->GetLocalEndpoint().ToString();
+	}
+
+	return FString(); 
 }
 
 bool FStormSyncTransportTcpServer::IsActive() const
@@ -236,7 +267,7 @@ void FStormSyncTransportTcpServer::HandleIncomingBuffer(const FIPv4Endpoint& InE
 
 void FStormSyncTransportTcpServer::HandleReceivedBuffer(const FIPv4Endpoint& InEndpoint, const TSharedPtr<FSocket>& InClientSocket, const TArray<uint8>& InBytes)
 {
-	UE_LOG(LogStormSyncServer, Display, TEXT("Handle received buffer via %s:%d (Buffer Size: %d)"), *InEndpoint.Address.ToString(), InEndpoint.Port, InBytes.Num());
+	UE_LOG(LogStormSyncServer, Display, TEXT("Handle received buffer via %s (Buffer Size: %d)"), *InEndpoint.ToString(), InBytes.Num());
 	ReceivedBufferEvent.Broadcast(InEndpoint, InClientSocket, MakeShared<TArray<uint8>>(InBytes));
 
 	// Fully received buffer, clean up connection and the state associated with it
@@ -300,3 +331,11 @@ bool FStormSyncTransportTcpServer::SendMessage(const FString& InMessage, const F
 	UE_LOG(LogStormSyncServer, Verbose, TEXT("Trying to send message to disconnected client %s - Message: %s"), *InEndpoint.ToString(), *InMessage);
 	return false;
 }
+
+void FStormSyncTransportTcpServer::FSocketDeleter::operator()(FSocket* InSocket) const
+{
+	if (InSocket)
+	{
+		ISocketSubsystem::Get(PLATFORM_SOCKETSUBSYSTEM)->DestroySocket(InSocket);
+	}
+}
\ No newline at end of file
diff --git a/Engine/Plugins/Experimental/StormSync/Source/StormSyncTransportServer/Private/Socket/StormSyncTransportTcpServer.h b/Engine/Plugins/Experimental/StormSync/Source/StormSyncTransportServer/Private/Socket/StormSyncTransportTcpServer.h
--- a/Engine/Plugins/Experimental/StormSync/Source/StormSyncTransportServer/Private/Socket/StormSyncTransportTcpServer.h
+++ b/Engine/Plugins/Experimental/StormSync/Source/StormSyncTransportServer/Private/Socket/StormSyncTransportTcpServer.h
@@ -103,7 +103,7 @@ public:
 	 */
 	void StopListening();
 
-	/** Endpoint address we're currently listening on (ip:port) */
+	/** Returns endpoint address we're currently listening on (ip:port) or empty string if not active. */
 	FString GetEndpointAddress() const;
 
 	/** Returns whether underlying FTcpListener has been created, is currently active and listening for incoming connections */
@@ -113,8 +113,14 @@ private:
 	/** The amount of time in seconds to consider a client to be inactive, if it has not sent data in this interval. */
 	uint32 DefaultInactiveTimeoutSeconds;
 
-	/** The endpoint used when creating the tcp listener */
-	TUniquePtr<FIPv4Endpoint> Endpoint;
+	/** The desired endpoint for creating the tcp listener. The actual endpoint may be different in case of collision. */
+	FIPv4Endpoint DesiredEndpoint;
+	
+	struct FSocketDeleter
+	{
+		void operator()(FSocket* InSocket) const;
+	};
+	TUniquePtr<FSocket, FSocketDeleter> Socket;
 
 	/** Our instance of tcp listener, created in StartListening() */
 	TUniquePtr<FTcpListener> SocketListener;
@@ -128,12 +134,6 @@ private:
 	/** Map of currently connected and active clients */
 	TMap<FIPv4Endpoint, TUniquePtr<FClientConnection>> ClientConnections;
 
-	/** The ip address for our endpoint and tcp listener */
-	FIPv4Address EndpointAddress;
-	
-	/** The port for our endpoint and tcp listener */
-	uint16 EndpointPort;
-
 	/** For the thread */
 	std::atomic<bool> bStopping = false;
 
diff --git a/Engine/Plugins/Experimental/StormSync/Source/StormSyncTransportServer/Private/StormSyncTransportServerEndpoint.cpp b/Engine/Plugins/Experimental/StormSync/Source/StormSyncTransportServer/Private/StormSyncTransportServerEndpoint.cpp
--- a/Engine/Plugins/Experimental/StormSync/Source/StormSyncTransportServer/Private/StormSyncTransportServerEndpoint.cpp
+++ b/Engine/Plugins/Experimental/StormSync/Source/StormSyncTransportServer/Private/StormSyncTransportServerEndpoint.cpp
@@ -250,7 +250,7 @@ void FStormSyncTransportServerEndpoint::SendSyncResponse(const FStormSyncTranspo
 		Message->Modifiers = FileModifiers;
 
 		Message->HostName = FStormSyncTransportNetworkUtils::GetServerName();
-		Message->HostAddress = FStormSyncTransportNetworkUtils::GetTcpEndpointAddress();
+		Message->HostAddress = FStormSyncTransportNetworkUtils::GetCurrentTcpServerEndpointAddress();
 		Message->HostAdapterAddresses = FStormSyncTransportNetworkUtils::GetLocalAdapterAddresses();
 
 		UE_LOG(LogStormSyncServer, Display, TEXT("FStormSyncTransportServerEndpoint::SendSyncResponse - Sending back to %s sync response: %s"), *RemoteAddress.ToString(), *Message->ToString());
diff --git a/Engine/Plugins/Experimental/StormSync/Source/StormSyncTransportServer/Private/StormSyncTransportServerModule.cpp b/Engine/Plugins/Experimental/StormSync/Source/StormSyncTransportServer/Private/StormSyncTransportServerModule.cpp
--- a/Engine/Plugins/Experimental/StormSync/Source/StormSyncTransportServer/Private/StormSyncTransportServerModule.cpp
+++ b/Engine/Plugins/Experimental/StormSync/Source/StormSyncTransportServer/Private/StormSyncTransportServerModule.cpp
@@ -2,6 +2,7 @@
 
 #include "StormSyncTransportServerModule.h"
 
+#include "IStormSyncTransportCoreModule.h"
 #include "IStormSyncTransportLocalEndpoint.h"
 #include "MessageEndpoint.h"
 #include "Misc/CoreDelegates.h"
@@ -31,10 +32,19 @@ void FStormSyncTransportServerModule::StartupModule()
 
 	// Register for engine initialization completed so we can broadcast presence over the network and start heartbeats
 	FCoreDelegates::OnPostEngineInit.AddRaw(this, &FStormSyncTransportServerModule::OnPostEngineInit);
+
+	IStormSyncTransportCoreModule::Get().OnGetCurrentTcpServerEndpointAddress().BindRaw(this, &FStormSyncTransportServerModule::GetCurrentTcpServerEndpointAddress);
+	IStormSyncTransportCoreModule::Get().OnGetServerEndpointMessageAddress().BindRaw(this, &FStormSyncTransportServerModule::GetServerEndpointMessageAddressId);
 }
 
 void FStormSyncTransportServerModule::ShutdownModule()
 {
+	if (IStormSyncTransportCoreModule::IsAvailable())
+	{
+		IStormSyncTransportCoreModule::Get().OnGetCurrentTcpServerEndpointAddress().Unbind();
+		IStormSyncTransportCoreModule::Get().OnGetServerEndpointMessageAddress().Unbind();
+	}
+
 	FCoreDelegates::OnPostEngineInit.RemoveAll(this);
 
 	if (ServerEndpoint.IsValid())
@@ -146,6 +156,11 @@ bool FStormSyncTransportServerModule::GetServerStatus(FText& OutStatusText) cons
 	return bIsRunning;
 }
 
+FString FStormSyncTransportServerModule::GetCurrentTcpServerEndpointAddress() const
+{
+	return ServerEndpoint ? ServerEndpoint->GetTcpServerEndpointAddress() : FString();
+}
+
 void FStormSyncTransportServerModule::ConditionalStartHeartBeatEmitter()
 {
 	if (!HeartbeatEmitter)
diff --git a/Engine/Plugins/Experimental/StormSync/Source/StormSyncTransportServer/Private/StormSyncTransportServerModule.h b/Engine/Plugins/Experimental/StormSync/Source/StormSyncTransportServer/Private/StormSyncTransportServerModule.h
--- a/Engine/Plugins/Experimental/StormSync/Source/StormSyncTransportServer/Private/StormSyncTransportServerModule.h
+++ b/Engine/Plugins/Experimental/StormSync/Source/StormSyncTransportServer/Private/StormSyncTransportServerModule.h
@@ -15,7 +15,7 @@ public:
 	virtual void StartupModule() override;
 	virtual void ShutdownModule() override;
 
-	//~ Begin IStormSyncTransportClientModule
+	//~ Begin IStormSyncTransportServerModule
 	virtual void StartDiscoveryManager() override;
 	virtual void StartServerEndpoint(const FString& InEndpointFriendlyName) override;
 	virtual TSharedPtr<IStormSyncTransportServerLocalEndpoint> CreateServerLocalEndpoint(const FString& InEndpointFriendlyName) const override;
@@ -24,9 +24,12 @@ public:
 	virtual FStormSyncHeartbeatEmitter& GetHeartbeatEmitter() const override;
 	virtual bool IsRunning() const override;
 	virtual bool GetServerStatus(FText& OutStatusText) const override;
-	//~ End IStormSyncTransportClientModule
+	//~ End IStormSyncTransportServerModule
 
 private:
+	/** Returns the current address (ip:port) the tcp server is bound to. */
+	FString GetCurrentTcpServerEndpointAddress() const;
+	
 	/** Starts the heart beat emitter if it was not started already. */
 	void ConditionalStartHeartBeatEmitter();
 
diff --git a/Engine/Plugins/Experimental/StormSync/Source/StormSyncTransportServer/StormSyncTransportServer.Build.cs b/Engine/Plugins/Experimental/StormSync/Source/StormSyncTransportServer/StormSyncTransportServer.Build.cs
--- a/Engine/Plugins/Experimental/StormSync/Source/StormSyncTransportServer/StormSyncTransportServer.Build.cs
+++ b/Engine/Plugins/Experimental/StormSync/Source/StormSyncTransportServer/StormSyncTransportServer.Build.cs
@@ -23,8 +23,7 @@ namespace UnrealBuildTool.Rules
 					"Networking",
 					"Sockets",
 					"StormSyncCore",
-					"StormSyncTransportCore",
-					"StormSyncTransportClient",
+					"StormSyncTransportCore"
 				}
 			);
 
