[StormSync] Lifting the 2 GB size limit for SPAKs and TCP transfers

- Changing the APIs to use FArchive instead of TArray.
- Changing some of the remaining TArray to TArray64.
- No longer loading the entire files in memory, the files are copied in chunk buffers instead.

diff --git a/Engine/Plugins/Experimental/StormSync/Source/StormSyncCore/Internal/StormSyncMemoryReader.h b/Engine/Plugins/Experimental/StormSync/Source/StormSyncCore/Internal/StormSyncMemoryReader.h
new file mode 100644
--- /dev/null
+++ b/Engine/Plugins/Experimental/StormSync/Source/StormSyncCore/Internal/StormSyncMemoryReader.h
@@ -0,0 +1,54 @@
+// Copyright Epic Games, Inc. All Rights Reserved.
+
+#pragma once
+
+#include "HAL/UnrealMemory.h"
+#include "Serialization/MemoryArchive.h"
+#include "StormSyncCommonTypes.h"
+
+/**
+ * Archive for reading arbitrary data from a shared storm sync buffer.
+ */
+class FStormSyncMemoryReader : public FMemoryArchive
+{
+public:
+	//~ Begin FArchive
+	virtual FString GetArchiveName() const override
+	{
+		return TEXT("FStormSyncMemoryReader");
+	}
+
+	virtual int64 TotalSize() override
+	{
+		return Buffer ? Buffer->Num() : 0;
+	}
+
+	virtual void Serialize(void* OutData, int64 InNum) override
+	{
+		if (InNum && !IsError())
+		{
+			// Only serialize if we have the requested amount of data
+			if (Buffer && Offset + InNum <= TotalSize())
+			{
+				FMemory::Memcpy(OutData, &(*Buffer)[Offset], InNum);
+				Offset += InNum;
+			}
+			else
+			{
+				SetError();
+			}
+		}
+	}
+	//~ End FArchive
+
+	explicit FStormSyncMemoryReader(const FStormSyncBufferPtr& InBuffer)
+		: Buffer(InBuffer)
+	{
+		this->SetIsLoading(true);
+		this->SetIsPersistent(false);
+	}
+
+private:
+	/** Keep a shared reference on the buffer. */
+	FStormSyncBufferPtr Buffer;
+};
\ No newline at end of file
diff --git a/Engine/Plugins/Experimental/StormSync/Source/StormSyncCore/Private/StormSyncCoreUtils.cpp b/Engine/Plugins/Experimental/StormSync/Source/StormSyncCore/Private/StormSyncCoreUtils.cpp
--- a/Engine/Plugins/Experimental/StormSync/Source/StormSyncCore/Private/StormSyncCoreUtils.cpp
+++ b/Engine/Plugins/Experimental/StormSync/Source/StormSyncCore/Private/StormSyncCoreUtils.cpp
@@ -6,6 +6,7 @@
 #include "AssetRegistry/AssetData.h"
 #include "AssetRegistry/AssetRegistryModule.h"
 #include "HAL/FileManager.h"
+#include "HAL/FileManagerGeneric.h"
 #include "Misc/PackageName.h"
 #include "Misc/Paths.h"
 #include "Serialization/MemoryReader.h"
@@ -17,6 +18,38 @@
 
 #define LOCTEXT_NAMESPACE "StormSyncCoreUtils"
 
+namespace UE::StormSyncCore::Private
+{
+	// For now, we use the same buffer size as the file writer.
+	// Todo: It might be beneficial to have an implementation for uncached reads/writes
+	// for large file operations which may benefit from buffers that are more aligned with
+	// physical device's page size.
+	constexpr uint64 ChunkBufferSize = PLATFORM_FILE_WRITER_BUFFER_SIZE;
+
+	uint64 GetChunkBufferSize()
+	{
+		return ChunkBufferSize;
+	}
+
+	/** Utility function to copy from one archive to another using a temporary buffer. */
+	void ArchiveCopy(FArchive& InReadArchive, FArchive& InWriteArchive, int64 InSizeToCopy, TArray64<uint8>& InChunkBuffer)
+	{
+		// Note: This should already be validated prior, but we ensure here as a precaution.
+		if(!ensure(InReadArchive.IsLoading()) || !ensure(InWriteArchive.IsSaving()))
+		{
+			return;
+		}
+		
+		while (InSizeToCopy > 0)
+		{
+			const int64 ChunkSize = FMath::Min(InChunkBuffer.Num(), InSizeToCopy);
+			InReadArchive.Serialize(InChunkBuffer.GetData(), ChunkSize);
+			InWriteArchive.Serialize(InChunkBuffer.GetData(), ChunkSize);
+			InSizeToCopy -= ChunkSize;
+		}
+	}
+}
+
 bool FStormSyncCoreUtils::GetAssetData(const FString& InPackageName, TArray<FAssetData>& OutAssets, TArray<FName>& OutDependencies)
 {
 	FString Filename = InPackageName;
@@ -179,7 +212,7 @@ FStormSyncFileDependency FStormSyncCoreUtils::CreateStormSyncFile(const FName& I
 	return FileDependency;
 }
 
-bool FStormSyncCoreUtils::CreatePakBufferWithDependencies(const TArray<FName>& InPackageNames, TArray<uint8>& OutPakBuffer, FText& OutErrorText, const FOnFileAdded& InOnFileAdded)
+bool FStormSyncCoreUtils::CreatePakBufferWithDependencies(const TArray<FName>& InPackageNames, FArchive& OutPakArchive, FText& OutErrorText, const FOnFileAdded& InOnFileAdded)
 {
 	TArray<FName> PackageDependencies;
 	if (!GetDependenciesForPackages(InPackageNames, PackageDependencies, OutErrorText))
@@ -187,10 +220,10 @@ bool FStormSyncCoreUtils::CreatePakBufferWithDependencies(const TArray<FName>& I
 		return false;
 	}
 
-	return CreatePakBuffer(PackageDependencies, OutPakBuffer, OutErrorText, InOnFileAdded);
+	return CreatePakBuffer(PackageDependencies, OutPakArchive, OutErrorText, InOnFileAdded);
 }
 
-bool FStormSyncCoreUtils::CreatePakBuffer(const TArray<FName>& InPackageNames, TArray<uint8>& OutPakBuffer, FText& OutErrorText, const FOnFileAdded& InOnFileAdded)
+bool FStormSyncCoreUtils::CreatePakBuffer(const TArray<FName>& InPackageNames, FArchive& OutPakArchive, FText& OutErrorText, const FOnFileAdded& InOnFileAdded)
 {
 	const double StartTime = FPlatformTime::Seconds();
 
@@ -201,6 +234,12 @@ bool FStormSyncCoreUtils::CreatePakBuffer(const TArray<FName>& InPackageNames, T
 		return false;
 	}
 
+	if (!OutPakArchive.IsSaving())
+	{
+		OutErrorText = FText::FromString(FString::Printf(TEXT("CreatePakBuffer: Archive \"%s\" is not writable."), *OutPakArchive.GetArchiveName()));
+		return false;
+	}
+
 	TArray<FName> PackageNames;
 
 	const UStormSyncCoreSettings* Settings = GetDefault<UStormSyncCoreSettings>();
@@ -229,20 +268,16 @@ bool FStormSyncCoreUtils::CreatePakBuffer(const TArray<FName>& InPackageNames, T
 		return false;
 	}
 
-	TUniquePtr<FArchive> MemoryArchive = MakeUnique<FMemoryWriter>(OutPakBuffer);
-
-	// Create Pak
-	if (!MemoryArchive.IsValid())
-	{
-		OutErrorText = LOCTEXT("CreatePakBuffer_Invalid_MemoryHandle", "Unable to create in memory archive");
-		return false;
-	}
-
+	const int64 StartOffset = OutPakArchive.Tell();
+	
 	int32 FileCount = PackageNames.Num();
-	*MemoryArchive << FileCount;
+	OutPakArchive << FileCount;
 
-	UE_LOG(LogStormSyncCore, Display, TEXT("FStormSyncCoreUtils::CreatePakBuffer - Creating Pak file for %d files."), PackageNames.Num());
+	TArray64<uint8> ChunkBuffer;
+	ChunkBuffer.SetNumUninitialized(UE::StormSyncCore::Private::GetChunkBufferSize());
 
+	UE_LOG(LogStormSyncCore, Display, TEXT("FStormSyncCoreUtils::CreatePakBuffer - Creating Pak file for %d files."), PackageNames.Num());
+	
 	for (const FName& PackageName : PackageNames)
 	{
 		const FString PackageNameStr = PackageName.ToString();
@@ -264,27 +299,21 @@ bool FStormSyncCoreUtils::CreatePakBuffer(const TArray<FName>& InPackageNames, T
 			FMD5Hash FileMD5Hash = FMD5Hash::HashFile(*PackageFilepath);
 			FString FileHash = LexToString(FileMD5Hash);
 
-			// Create buffer
-			uint8* FileBuffer = static_cast<uint8*>(FMemory::Malloc(FileSize));
-			FileHandle->Serialize(FileBuffer, FileSize);
-
 			// Write Package path
 			FString FinalAssetName = PackageNameStr + FPaths::GetExtension(PackageFilepath, /*bIncludeDot*/true);
-			*MemoryArchive << FinalAssetName;
+			OutPakArchive << FinalAssetName;
 
 			// Write size of the buffer
-			*MemoryArchive << FileSize;
+			OutPakArchive << FileSize;
 
 			// Write file timestamp
-			*MemoryArchive << Timestamp;
+			OutPakArchive << Timestamp;
 
 			// Write file hash
-			*MemoryArchive << FileHash;
+			OutPakArchive << FileHash;
 
-			// Write to pack
-			MemoryArchive->Serialize(FileBuffer, FileSize);
-
-			FMemory::Free(FileBuffer);
+			// Copy file content to pack
+			UE::StormSyncCore::Private::ArchiveCopy(*FileHandle, OutPakArchive, FileSize, ChunkBuffer);
 
 			// Close the file
 			FileHandle->Close();
@@ -293,40 +322,39 @@ bool FStormSyncCoreUtils::CreatePakBuffer(const TArray<FName>& InPackageNames, T
 		}
 	}
 
-	UE_LOG(LogStormSyncCore, Display, TEXT("Added %d files to %s, %lld bytes total, time %.2lfs."), FileCount, *MemoryArchive->GetArchiveName(), MemoryArchive->TotalSize(), FPlatformTime::Seconds() - StartTime);
-
-	MemoryArchive->Close();
-	MemoryArchive.Reset();
+	UE_LOG(LogStormSyncCore, Display, TEXT("Added %d files to %s, %lld bytes total, time %.2lfs."), FileCount, *OutPakArchive.GetArchiveName(), OutPakArchive.Tell() - StartOffset, FPlatformTime::Seconds() - StartTime);
 
 	return true;
 }
 
-bool FStormSyncCoreUtils::ExtractPakBuffer(const TArray<uint8>& InPakBuffer, const FStormSyncCoreExtractArgs& InExtractArgs, TMap<FString, FString>& OutSuccessfullyExtractedPackages, TArray<FText>& OutErrors)
+bool FStormSyncCoreUtils::ExtractPakBuffer(FArchive& InPakArchive, const FStormSyncCoreExtractArgs& InExtractArgs, TArray<FText>& OutErrors)
 {
-	UE_LOG(LogStormSyncCore, Display, TEXT("FStormSyncCoreUtils:ExtractPakBuffer Extracting package from buffer of size: %d"), InPakBuffer.Num());
-
-	const TUniquePtr<FArchive> MemoryArchive = MakeUnique<FMemoryReader>(InPakBuffer);
-	if (!MemoryArchive.IsValid())
+	if (!InPakArchive.IsLoading())
 	{
-		OutErrors.Add(LOCTEXT("Invalid_PakHandle", "Unable to create pak buffer archive"));
+		UE_LOG(LogStormSyncCore, Error, TEXT("FStormSyncCoreUtils:ExtractPakBuffer Cannot Extract package from non-reader archive \"%s\""), *InPakArchive.GetArchiveName());
 		return false;
 	}
+	
+	UE_LOG(LogStormSyncCore, Display, TEXT("FStormSyncCoreUtils:ExtractPakBuffer Extracting package from archive of size: %lld"), InPakArchive.TotalSize());
 
 	bool bSuccess = true;
 
 	int32 FileCount = 0;
-	*MemoryArchive << FileCount;
+	InPakArchive << FileCount;
 
 	// Notify pack extraction process is starting
 	InExtractArgs.OnPakPreExtract.ExecuteIfBound(FileCount);
 
 	UE_LOG(LogStormSyncCore, Verbose, TEXT("FStormSyncCoreUtils:ExtractPakBuffer FileCount: %d"), FileCount);
 
+	TArray64<uint8> ChunkBuffer;
+	ChunkBuffer.SetNumUninitialized(UE::StormSyncCore::Private::GetChunkBufferSize());
+
 	for (int32 Index = 0; Index < FileCount; ++Index)
 	{
 		// Extract package path
 		FString PackagePath;
-		*MemoryArchive << PackagePath;
+		InPakArchive << PackagePath;
 
 		// We want original PackageName (as right now, pak creation stores the extension in pak file)
 		const FString Extension = FPaths::GetExtension(PackagePath, /*bIncludeDot*/true);
@@ -334,24 +362,15 @@ bool FStormSyncCoreUtils::ExtractPakBuffer(const TArray<uint8>& InPakBuffer, con
 
 		// Extract file Size
 		uint64 FileSize = 0;
-		*MemoryArchive << FileSize;
+		InPakArchive << FileSize;
 
 		// Extract file timestamp
 		int64 Timestamp = 0;
-		*MemoryArchive << Timestamp;
+		InPakArchive << Timestamp;
 
 		// Extract file hash
 		FString FileHash;
-		*MemoryArchive << FileHash;
-
-		// Load remaining data (raw buffer of the file itself)
-		TSharedPtr<TArray<uint8>, ESPMode::ThreadSafe> BufferPtr = MakeShareable(new TArray<uint8>());
-
-		TArray<uint8>& Bytes = *BufferPtr.Get();
-		Bytes.Reset(FileSize);
-		Bytes.AddUninitialized(FileSize);
-
-		MemoryArchive->Serialize(Bytes.GetData(), Bytes.Num());
+		InPakArchive << FileHash;
 
 		// Figure out file destination on disk
 		const FName PackageName = FName(*PackagePath);
@@ -363,16 +382,49 @@ bool FStormSyncCoreUtils::ExtractPakBuffer(const TArray<uint8>& InPakBuffer, con
 			// We were not able to determine destination output, mark as errored.
 			OutErrors.Add(ErrorText);
 			bSuccess = false;
+			InPakArchive.Seek(InPakArchive.Tell() + FileSize);	// Skip file content.
 		}
 		else
 		{
 			DestFilepath += Extension;
 
-			OutSuccessfullyExtractedPackages.Add(PackagePath, DestFilepath);
-
 			// Notify pack extraction for individual files
 			FStormSyncFileDependency FileDependency = FStormSyncFileDependency(PackageName, FileSize, Timestamp, FileHash);
-			InExtractArgs.OnFileExtract.ExecuteIfBound(FileDependency, DestFilepath, BufferPtr);
+
+			if (InExtractArgs.OnGetArchiveForExtract.IsBound())
+			{
+				const FStormSyncArchivePtr Archive = InExtractArgs.OnGetArchiveForExtract.Execute(FileDependency, DestFilepath);
+				if (Archive && Archive->IsSaving())
+				{
+					UE::StormSyncCore::Private::ArchiveCopy(InPakArchive, *Archive, FileSize, ChunkBuffer);
+
+					// Call the end of extraction with the archive.
+					InExtractArgs.OnArchiveExtracted.ExecuteIfBound(FileDependency, DestFilepath, Archive);
+				}
+				else
+				{
+					if (Archive && !Archive->IsSaving())
+					{
+						const FString ErrorMessage = FString::Printf(TEXT("FStormSyncCoreUtils::ExtractPakBuffer - A non-saving archive was provided, skipping \"%s\"."), *DestFilepath); 
+						OutErrors.Add(FText::FromString(ErrorMessage));
+						bSuccess = false;
+					}
+					InPakArchive.Seek(InPakArchive.Tell() + FileSize);	// Skip file content.
+				}
+			}
+			else if (InExtractArgs.OnFileExtract.IsBound())
+			{
+				// Load remaining data (raw buffer of the file itself)
+				FStormSyncBufferPtr Buffer = MakeShared<FStormSyncBuffer>();
+				Buffer->Reset(FileSize);
+				Buffer->AddUninitialized(FileSize);
+				InPakArchive.Serialize(Buffer->GetData(), Buffer->Num());
+				InExtractArgs.OnFileExtract.Execute(FileDependency, DestFilepath, Buffer);
+			}
+			else
+			{
+				InPakArchive.Seek(InPakArchive.Tell() + FileSize);	// Skip file content.
+			}
 		}
 	}
 
diff --git a/Engine/Plugins/Experimental/StormSync/Source/StormSyncCore/Public/StormSyncCommonTypes.h b/Engine/Plugins/Experimental/StormSync/Source/StormSyncCore/Public/StormSyncCommonTypes.h
--- a/Engine/Plugins/Experimental/StormSync/Source/StormSyncCore/Public/StormSyncCommonTypes.h
+++ b/Engine/Plugins/Experimental/StormSync/Source/StormSyncCore/Public/StormSyncCommonTypes.h
@@ -5,8 +5,12 @@
 #include "Templates/SharedPointer.h"
 #include "StormSyncCommonTypes.generated.h"
 
+class FArchive;
+
 // Common alias to a ThreadSafe SharedPtr to hold a Buffer (array of uint8)
-using FStormSyncBufferPtr = TSharedPtr<TArray<uint8>, ESPMode::ThreadSafe>;
+using FStormSyncBuffer = TArray64<uint8>;
+using FStormSyncBufferPtr = TSharedPtr<FStormSyncBuffer, ESPMode::ThreadSafe>;
+using FStormSyncArchivePtr = TSharedPtr<FArchive, ESPMode::ThreadSafe>;
 
 /** Engine type values for a storm sync connected device */
 UENUM()
diff --git a/Engine/Plugins/Experimental/StormSync/Source/StormSyncCore/Public/StormSyncCoreDelegates.h b/Engine/Plugins/Experimental/StormSync/Source/StormSyncCore/Public/StormSyncCoreDelegates.h
--- a/Engine/Plugins/Experimental/StormSync/Source/StormSyncCore/Public/StormSyncCoreDelegates.h
+++ b/Engine/Plugins/Experimental/StormSync/Source/StormSyncCore/Public/StormSyncCoreDelegates.h
@@ -72,9 +72,9 @@ struct STORMSYNCCORE_API FStormSyncCoreDelegates
 	 * Delegate type for when an incoming pak extraction process starts
 	 *
 	 * @param PackageDescriptor Metadata information about incoming buffer (name, description, etc.)
-	 * @param FileBuffer The raw file buffer (as a  shared ptr array of bytes) for the incoming pak buffer
+	 * @param Archive The archive for the incoming pak buffer
 	 */
-	DECLARE_MULTICAST_DELEGATE_TwoParams(FOnRequestImportBuffer, const FStormSyncPackageDescriptor&, const FStormSyncBufferPtr&);
+	DECLARE_MULTICAST_DELEGATE_TwoParams(FOnRequestImportBuffer, const FStormSyncPackageDescriptor&, const FStormSyncArchivePtr&);
 
 	/**
 	 * Delegate type for when one an incoming pak extraction process starts
diff --git a/Engine/Plugins/Experimental/StormSync/Source/StormSyncCore/Public/StormSyncCoreUtils.h b/Engine/Plugins/Experimental/StormSync/Source/StormSyncCore/Public/StormSyncCoreUtils.h
--- a/Engine/Plugins/Experimental/StormSync/Source/StormSyncCore/Public/StormSyncCoreUtils.h
+++ b/Engine/Plugins/Experimental/StormSync/Source/StormSyncCore/Public/StormSyncCoreUtils.h
@@ -15,36 +15,56 @@
  */
 struct FStormSyncCoreExtractArgs
 {
-	/** Delegate type for when one an incoming pak extraction process starts */
-	DECLARE_DELEGATE_OneParam(FOnPakPreExtract, int32);
-
-	/** Delegate type for when one an incoming pak extraction is ended */
-	DECLARE_DELEGATE_OneParam(FOnPakPostExtract, int32);
-	
-	/** Delegate type for when one asset is being extracted from an incoming pak */
-	DECLARE_DELEGATE_ThreeParams(FOnFileExtract, const FStormSyncFileDependency&, FString, const FStormSyncBufferPtr&);
-
 	/**
 	 * Delegate type for when one an incoming pak extraction process starts
-	 *
 	 * @param FileCount The int32 file count extracted from incoming buffer
 	 */
-	FOnPakPreExtract OnPakPreExtract;
+	DECLARE_DELEGATE_OneParam(FOnPakPreExtract, int32 /*InFileCount*/);
 
 	/**
 	 * Delegate type for when one an incoming pak extraction is ended
-	 *
 	 * @param FileCount The int32 file count extracted from incoming buffer
 	 */
-	FOnPakPostExtract OnPakPostExtract;
+	DECLARE_DELEGATE_OneParam(FOnPakPostExtract, int32 /*InFileCount*/);
+
+	/**
+	 * Delegate type called to obtain an archive to extract an asset. This is called before extraction starts.
+	 * @param InFileDependency struct holding info such as Package Name, File size, timestamp and hash
+	 * @param InDestFilename The fully qualified destination filepath where the file should be extracted to
+	 * @return Archive (writer) into which the incoming file will be written.
+	 */
+	DECLARE_DELEGATE_RetVal_TwoParams(FStormSyncArchivePtr, FOnGetArchiveForExtract, const FStormSyncFileDependency& /*InFileDependency*/, const FString& /*InDestFilename*/);
+
+	/**
+	 * Delegate type called after a file buffer is extracted from an incoming pak into the provided archive.
+ 	 * @param InFileDependency struct holding info such as Package Name, File size, timestamp and hash
+	 * @param InDestFilename The fully qualified destination filepath where the file should be extracted to
+	 * @param InArchive The archive that was returned by FOnGetArchiveForExtract is returned back.
+	 */
+	DECLARE_DELEGATE_ThreeParams(FOnArchiveExtracted, const FStormSyncFileDependency& /*InFileDependency*/, const FString& /*InDestFilename*/, const FStormSyncArchivePtr& /*InArchive*/);
 	
 	/**
-	 * Delegate type triggered when a file asset is extracted from an incoming pak
-	 *
+	 * Delegate type for when one asset is being extracted from an incoming pak
 	 * @param FileDependency struct holding info such as Package Name, File size, timestamp and hash
-	 * @param DestFilepath The fully qualified destination filepath where the file should be extracted to
-	 * @param FileBuffer The raw file buffer (as an array of bytes shared ptr) for the individual files extracted from incoming pak buffer
-	 */
+	 * @param InDestFilename The fully qualified destination filepath where the file should be extracted to
+	 * @param InBuffer The raw buffer (as an array of bytes shared ptr) for the individual files extracted from incoming pak buffer
+	 * @remark if an archive was provided by OnGetArchiveForExtract delegate, this delegate will not be called.
+	*/
+	DECLARE_DELEGATE_ThreeParams(FOnFileExtract, const FStormSyncFileDependency& /*FileDependency*/, FString /*InDestFilename*/, const FStormSyncBufferPtr& /*InBuffer*/);
+	
+	/** Delegate type for when one an incoming pak extraction process starts */
+	FOnPakPreExtract OnPakPreExtract;
+
+	/** Delegate type for when one an incoming pak extraction is ended */
+	FOnPakPostExtract OnPakPostExtract;
+
+	/** Delegate type called to obtain an archive to extract an asset. This is called before extraction starts. */
+	FOnGetArchiveForExtract OnGetArchiveForExtract;
+
+	/** Delegate type called after a file buffer is extracted from the archive. */
+	FOnArchiveExtracted OnArchiveExtracted;
+
+	/** Delegate type triggered when a file asset is extracted from an incoming pak */
 	FOnFileExtract OnFileExtract;
 
 	/** Default constructor */
@@ -149,7 +169,7 @@ public:
 	 * @return Returns true if we were able to generate the pak buffer false if the operation failed.
 	 * Note that in case of failure, ErrorText is filled with further information.
 	 */
-	STORMSYNCCORE_API static bool CreatePakBufferWithDependencies(const TArray<FName>& InPackageNames, TArray<uint8>& OutPakBuffer, FText& OutErrorText, const FOnFileAdded& InOnFileAdded = FOnFileAdded());
+	STORMSYNCCORE_API static bool CreatePakBufferWithDependencies(const TArray<FName>& InPackageNames, FArchive& OutPakArchive, FText& OutErrorText, const FOnFileAdded& InOnFileAdded = FOnFileAdded());
 
 	/**
 	 * Attempts to create an ava pak buffer using a memory archive from the provided list of Package Names.
@@ -161,7 +181,7 @@ public:
 	 * use CreatePakBufferWithDependencies instead.
 	 *
 	 * @param InPackageNames List of package name to include in the pak
-	 * @param OutPakBuffer Raw buffer of the created in memory archive
+	 * @param OutPakArchive Destination archive
 	 * @param OutErrorText Localized Text if the operation failed
 	 * @param InOnFileAdded Delegate triggered with FStormSyncFileDependency for each added file
 	 * 
@@ -170,22 +190,21 @@ public:
 	 * @return Returns true if we were able to generate the pak buffer false if the operation failed.
 	 * Note that in case of failure, ErrorText is filled with further information.
 	 */
-	STORMSYNCCORE_API static bool CreatePakBuffer(const TArray<FName>& InPackageNames, TArray<uint8>& OutPakBuffer, FText& OutErrorText, const FOnFileAdded& InOnFileAdded = FOnFileAdded());
+	STORMSYNCCORE_API static bool CreatePakBuffer(const TArray<FName>& InPackageNames, FArchive& OutPakArchive, FText& OutErrorText, const FOnFileAdded& InOnFileAdded = FOnFileAdded());
 
 	/**
 	 * Main entry point for unpacking. Extracts an ava pak into the current Unreal project (Game content or Plugins content folder)
 	 * 
 	 * Will extract package names from pak buffer, and create any .uasset files that does not exist yet on local project.
 	 *
-	 * @param InPakBuffer Raw ava pak buffer to extract into project
+	 * @param InPakArchive Raw ava pak archive to extract into project
 	 * @param InExtractArgs The arguments used to trigger delegates during extraction process
-	 * @param OutSuccessfullyExtractedPackages Map of successfully imported files. Key is the package name for the file, Value is the output destination.
 	 * @param OutErrors List of localized Text if the operation failed
 	 *
 	 * @return Returns true if we were able to fully import the package and false if the operation failed (even for partial imports).
 	 * Note that in case of failure, Errors array of localized text is filled with further information.
 	 */
-	STORMSYNCCORE_API static bool ExtractPakBuffer(const TArray<uint8>& InPakBuffer, const FStormSyncCoreExtractArgs& InExtractArgs, TMap<FString, FString>& OutSuccessfullyExtractedPackages, TArray<FText>& OutErrors);
+	STORMSYNCCORE_API static bool ExtractPakBuffer(FArchive& InPakArchive, const FStormSyncCoreExtractArgs& InExtractArgs, TArray<FText>& OutErrors);
 
 	/**
 	 * Converts the provided bytes size into a human readable string.
diff --git a/Engine/Plugins/Experimental/StormSync/Source/StormSyncCore/Public/StormSyncImportTypes.h b/Engine/Plugins/Experimental/StormSync/Source/StormSyncCore/Public/StormSyncImportTypes.h
--- a/Engine/Plugins/Experimental/StormSync/Source/StormSyncCore/Public/StormSyncImportTypes.h
+++ b/Engine/Plugins/Experimental/StormSync/Source/StormSyncCore/Public/StormSyncImportTypes.h
@@ -4,7 +4,6 @@
 
 #include "Containers/UnrealString.h"
 #include "Internationalization/Text.h"
-#include "StormSyncCommonTypes.h"
 #include "StormSyncPackageDescriptor.h"
 #include "StormSyncImportTypes.generated.h"
 
@@ -24,21 +23,20 @@ struct FStormSyncImportFileInfo
 	/** The associated fully qualified path of the .uasset */
 	FString DestFilepath;
 
-	/** Holds a shared ptr to the serialized buffer */
-	FStormSyncBufferPtr FileBuffer = nullptr;
-
 	/** UI-exposed reason of the import */
 	FText ImportReason;
 	
 	/** UI-exposed extended info for reason of the import */
 	FText ImportReasonTooltip;
 
+	/** True if the file was a new asset */
+	bool bNewAsset = false;
+
 	FStormSyncImportFileInfo() = default;
 
-	FStormSyncImportFileInfo(const FStormSyncFileDependency& InFileDependency, const FString& InDestFilepath, const FStormSyncBufferPtr& InFileBuffer)
+	FStormSyncImportFileInfo(const FStormSyncFileDependency& InFileDependency, const FString& InDestFilepath)
 		: FileDependency(InFileDependency)
 		, DestFilepath(InDestFilepath)
-		, FileBuffer(InFileBuffer)
 	{
 	}
 };
diff --git a/Engine/Plugins/Experimental/StormSync/Source/StormSyncEditor/Private/ContextMenu/StormSyncAssetFolderContextMenu.cpp b/Engine/Plugins/Experimental/StormSync/Source/StormSyncEditor/Private/ContextMenu/StormSyncAssetFolderContextMenu.cpp
--- a/Engine/Plugins/Experimental/StormSync/Source/StormSyncEditor/Private/ContextMenu/StormSyncAssetFolderContextMenu.cpp
+++ b/Engine/Plugins/Experimental/StormSync/Source/StormSyncEditor/Private/ContextMenu/StormSyncAssetFolderContextMenu.cpp
@@ -777,15 +777,19 @@ void FStormSyncAssetFolderContextMenu::OnExportWizardCompleted(const TArray<FNam
 		// is false)
 		return;
 	}
-
 	
-	TArray<uint8> Buffer;
-
 	// Create buffer based on provided package names
 	{
 		FScopedSlowTask SlowTask(InPackageNames.Num(), LOCTEXT("ExportPackages_CreatingBuffer", "Creating Buffer..."));
 		SlowTask.MakeDialog();
 
+		TUniquePtr<FArchive> FileWriter(IFileManager::Get().CreateFileWriter(*InFilepath));
+		if (!FileWriter.IsValid())
+		{
+			UE_LOG(LogStormSyncEditor, Error, TEXT("FStormSyncAssetFolderContextMenu::OnExportWizardCompleted - Error creating file writer"));
+			return;
+		}
+		
 		const FStormSyncCoreUtils::FOnFileAdded Delegate = FStormSyncCoreUtils::FOnFileAdded::CreateLambda([&SlowTask](const FStormSyncFileDependency& FileDependency)
 		{
 			SlowTask.EnterProgressFrame();
@@ -793,26 +797,15 @@ void FStormSyncAssetFolderContextMenu::OnExportWizardCompleted(const TArray<FNam
 	
 		FText ErrorText;
 		TArray<FStormSyncFileDependency> PackedFiles;
-		if (!FStormSyncCoreUtils::CreatePakBuffer(InPackageNames, Buffer, ErrorText, Delegate))
+		if (!FStormSyncCoreUtils::CreatePakBuffer(InPackageNames, *FileWriter, ErrorText, Delegate))
 		{
-			UE_LOG(LogStormSyncEditor, Error, TEXT("FStormSyncAssetFolderContextMenu::OnExportWizardCompleted - CreatePakBufferWithDependencies Error: %s"), *ErrorText.ToString());
+			FileWriter.Reset();
+			IFileManager::Get().Delete(*InFilepath);
+			UE_LOG(LogStormSyncEditor, Error, TEXT("FStormSyncAssetFolderContextMenu::OnExportWizardCompleted - CreatePakBuffer Error: %s"), *ErrorText.ToString());
 			return;
 		}
 	}
 
-	// Writing buffer on disk
-	{
-		const TUniquePtr<FArchive> FileWriter(IFileManager::Get().CreateFileWriter(*InFilepath));
-		if (!FileWriter.IsValid())
-		{
-			UE_LOG(LogStormSyncEditor, Error, TEXT("FStormSyncAssetFolderContextMenu::OnExportWizardCompleted - Error creating file writer"));
-			return;
-		}
-
-		FileWriter->Serialize(Buffer.GetData(), Buffer.Num());
-		FileWriter->Close();
-	}
-
 	// Notify user
 	UE_LOG(LogStormSyncEditor, Display, TEXT("FStormSyncAssetFolderContextMenu::OnExportWizardCompleted - Created file %s"), *InFilepath);
 
diff --git a/Engine/Plugins/Experimental/StormSync/Source/StormSyncImport/Private/Subsystems/StormSyncImportSubsystem.cpp b/Engine/Plugins/Experimental/StormSync/Source/StormSyncImport/Private/Subsystems/StormSyncImportSubsystem.cpp
--- a/Engine/Plugins/Experimental/StormSync/Source/StormSyncImport/Private/Subsystems/StormSyncImportSubsystem.cpp
+++ b/Engine/Plugins/Experimental/StormSync/Source/StormSyncImport/Private/Subsystems/StormSyncImportSubsystem.cpp
@@ -114,12 +114,11 @@ bool UStormSyncImportSubsystem::EnqueueImportTask(const TSharedPtr<IStormSyncImp
 bool UStormSyncImportSubsystem::PerformFileImport(const FString& InFilename)
 {
 	UE_LOG(LogStormSyncImport, Display, TEXT("FStormSyncEditorUtils::PerformImport for %s"), *InFilename);
-	TArray<uint8> Data;
 
-	const FStormSyncBufferPtr Buffer = MakeShareable(new TArray<uint8>);
-	if (!FFileHelper::LoadFileToArray(*Buffer.Get(), *InFilename))
+	const FStormSyncArchivePtr Reader = FStormSyncArchivePtr(IFileManager::Get().CreateFileReader(*InFilename));
+	if (!Reader)
 	{
-		UE_LOG(LogStormSyncImport, Error, TEXT("Failed to load file '%s' to array"), *InFilename);
+		UE_LOG(LogStormSyncImport, Error, TEXT("Failed to open file '%s'"), *InFilename);
 		return false;
 	}
 
@@ -128,14 +127,14 @@ bool UStormSyncImportSubsystem::PerformFileImport(const FString& InFilename)
 	PackageDescriptor.Name = FPaths::GetBaseFilename(InFilename);
 
 	constexpr bool bShowWizard = true;
-	return PerformImport(PackageDescriptor, Buffer, bShowWizard);
+	return PerformImport(PackageDescriptor, Reader, bShowWizard);
 }
 
-bool UStormSyncImportSubsystem::PerformBufferImport(const FStormSyncPackageDescriptor& InPackageDescriptor, const FStormSyncBufferPtr& InBuffer)
+bool UStormSyncImportSubsystem::PerformBufferImport(const FStormSyncPackageDescriptor& InPackageDescriptor, const FStormSyncArchivePtr& InArchive)
 {
-	if (!InBuffer.IsValid())
+	if (!InArchive.IsValid())
 	{
-		UE_LOG(LogStormSyncImport, Error, TEXT("UStormSyncImportSubsystem::PerformBufferImport failed to import from invalid buffer"));
+		UE_LOG(LogStormSyncImport, Error, TEXT("UStormSyncImportSubsystem::PerformBufferImport failed to import from invalid archive"));
 		return false;
 	}
 
@@ -144,137 +143,239 @@ bool UStormSyncImportSubsystem::PerformBufferImport(const FStormSyncPackageDescr
 	
 	const bool bTcpDryRun = Settings->IsTcpDryRun();
 	const bool bShowWizard = Settings->ShouldShowImportWizard();
-	return PerformImport(InPackageDescriptor, InBuffer, bShowWizard, bTcpDryRun);
+	return PerformImport(InPackageDescriptor, InArchive, bShowWizard, bTcpDryRun);
 }
 
-bool UStormSyncImportSubsystem::PerformImport(const FStormSyncPackageDescriptor& InPackageDescriptor, const FStormSyncBufferPtr& InBuffer, bool bShowWizard, bool bDryRun)
+bool UStormSyncImportSubsystem::ParsePak(const FStormSyncArchivePtr& InArchive, TArray<FStormSyncImportFileInfo>& OutFileInfos)
 {
-	if (!InBuffer.IsValid())
+	if (!InArchive.IsValid())
 	{
-		UE_LOG(LogStormSyncImport, Error, TEXT("UStormSyncImportSubsystem::PerformBufferImport failed to import from invalid buffer"));
+		UE_LOG(LogStormSyncImport, Error, TEXT("UStormSyncImportSubsystem::ParsePak failed to import from invalid archive"));
+		return false;
+	}
+
+	if (!InArchive->IsLoading())
+	{
+		UE_LOG(LogStormSyncImport, Error, TEXT("UStormSyncImportSubsystem::ParsePak archive \"%s\" is not loading"), *InArchive->GetArchiveName());
+		return false;
+	}
+
+	FStormSyncCoreExtractArgs ParseArgs;
+
+	ParseArgs.OnGetArchiveForExtract.BindLambda([&OutFileInfos](const FStormSyncFileDependency& InFileDependency, const FString& InDestFilepath) -> FStormSyncArchivePtr
+	{
+		OutFileInfos.Add(FStormSyncImportFileInfo(InFileDependency, InDestFilepath));
+		return nullptr;	// Returning null results in skipping the content of the file in the source archive.
+	});
+
+	TArray<FText> Errors;
+	if (!FStormSyncCoreUtils::ExtractPakBuffer(*InArchive.Get(), ParseArgs, Errors))
+	{
+		UE_LOG(LogStormSyncImport, Warning, TEXT("FStormSyncEditorUtils::PerformImport - Error extracting package ..."));
+		for (const FText& Error : Errors)
+		{
+			UE_LOG(LogStormSyncImport, Warning, TEXT("\t %s"), *Error.ToString());
+		}
+		return false;
+	}
+	return true;
+}
+
+
+bool UStormSyncImportSubsystem::PerformImport(const FStormSyncPackageDescriptor& InPackageDescriptor, const FStormSyncArchivePtr& InArchive, bool bShowWizard, bool bDryRun)
+{
+	if (!InArchive.IsValid())
+	{
+		UE_LOG(LogStormSyncImport, Error, TEXT("UStormSyncImportSubsystem::PerformImport failed to import from invalid archive"));
+		return false;
+	}
+
+	if (!InArchive->IsLoading())
+	{
+		UE_LOG(LogStormSyncImport, Error, TEXT("UStormSyncImportSubsystem::PerformBufferImport archive \"%s\" is not loading"), *InArchive->GetArchiveName());
 		return false;
 	}
 
 	// Package descriptor local copy
 	FStormSyncPackageDescriptor PackageDescriptor = InPackageDescriptor;
 	
-	float InitialAmountOfWork = 1.f;
-	FScopedSlowTask SlowTask(InitialAmountOfWork, LOCTEXT("FiguringOutFileStates", "Figuring out file states..."));
-	SlowTask.MakeDialog();
+	// The list of all files included in the buffer for filename
+	TArray<FStormSyncImportFileInfo> AllFileInfos;
+
+	// We have to do a 2 pass parsing so the archive needs to be seekable.
+	int64 StartPosition = InArchive->Tell();
+
+	if (StartPosition < 0)
+	{
+		UE_LOG(LogStormSyncImport, Error, TEXT("UStormSyncImportSubsystem::PerformBufferImport archive \"%s\" is not seekable"), *InArchive->GetArchiveName());
+		return false;
+	}
 
-	FStormSyncCoreExtractArgs ExtractArgs;
+	// Do a first pass to extract all the file information.
+	if (!ParsePak(InArchive, AllFileInfos))
+	{
+		return false;
+	}
+	
+	if (bDryRun)
+	{
+		UE_LOG(LogStormSyncImport, Display, TEXT("\tOnPakExtract - File Count: %d"), AllFileInfos.Num());		
+		for (const FStormSyncImportFileInfo& FileInfo : AllFileInfos)
+		{
+			UE_LOG(LogStormSyncImport, Display, TEXT("\tOnFileExtract"));
+			UE_LOG(LogStormSyncImport, Display, TEXT("\t\tPackageName: %s"), *FileInfo.FileDependency.PackageName.ToString());
+			UE_LOG(LogStormSyncImport, Display, TEXT("\t\tDestFilepath: %s"), *FileInfo.DestFilepath);
+			UE_LOG(LogStormSyncImport, Display, TEXT("\t\tFileSize: %lld"), FileInfo.FileDependency.FileSize);
+		}
+		return true;
+	}
 
 	// The list of files with changes detected (either size of hash changed)
 	TArray<FStormSyncImportFileInfo> FilesToImport;
 
-	// The list of all files included in the buffer for filename
-	TArray<FStormSyncImportFileInfo> BufferFiles;
-
-	ExtractArgs.OnPakPreExtract.BindLambda([&SlowTask, PackageDescriptor, bDryRun](const int32 FileCount)
 	{
-		SlowTask.TotalAmountOfWork = FileCount;
-		UE_LOG(LogStormSyncImport, Display, TEXT("\tOnPakPreExtract - File Count: %d"), FileCount);
+		FScopedSlowTask SlowTaskHash(AllFileInfos.Num(), LOCTEXT("FiguringOutFileStates", "Figuring out file states..."));
+		SlowTaskHash.MakeDialog();
 		
-		if (!bDryRun)
+		for (const FStormSyncImportFileInfo& FileInfo : AllFileInfos)
 		{
-			UStormSyncImportSubsystem::Get().HandlePakPreExtract(PackageDescriptor, FileCount);
+			SlowTaskHash.EnterProgressFrame();
+			FStormSyncImportFileInfo FileInfoToImport = FileInfo;
+			if (UStormSyncImportSubsystem::Get().ShouldFileBeImported(FileInfoToImport))
+			{
+				FileInfoToImport.bNewAsset = !FPaths::FileExists(FileInfo.DestFilepath);
+				FilesToImport.Add(MoveTemp(FileInfoToImport));
+			}
 		}
-	});
+	}
+
+	// Important: from this point on, FilesToImport can't be reallocated.
+	TArray<const FStormSyncImportFileInfo*> ExistingFilesToImport;
+	TArray<const FStormSyncImportFileInfo*> NewFilesToImport;
+	TMap<FString, const FStormSyncImportFileInfo*> FilesToImportMap;
+
+	ExistingFilesToImport.Reserve(FilesToImport.Num());
+	NewFilesToImport.Reserve(FilesToImport.Num());
+	FilesToImportMap.Reserve(FilesToImport.Num());
 	
-	ExtractArgs.OnPakPostExtract.BindLambda([PackageDescriptor, bShowWizard, bDryRun, &FilesToImport, &BufferFiles](const int32 FileCount)
+	for (const FStormSyncImportFileInfo& FileInfo : FilesToImport)
 	{
-		if (bDryRun)
+		if (FileInfo.bNewAsset)
 		{
-			UE_LOG(LogStormSyncImport, Display, TEXT("\tOnPakPostExtract - File Count: %d"), FileCount);
-			return;
+			NewFilesToImport.Add(&FileInfo);
 		}
+		else
+		{
+			ExistingFilesToImport.Add(&FileInfo);
+		}
+		FilesToImportMap.Add(FileInfo.DestFilepath, &FileInfo);
+	}
 
 #if WITH_EDITOR
-		if (bShowWizard)
-		{
-			// Create file import dialog. This is a modal dialog so it will return only after user selection
-			const TSharedRef<IStormSyncImportWizard> Wizard = FStormSyncEditorModule::Get().CreateWizard(FilesToImport, BufferFiles);
+	if (bShowWizard)
+	{
+		// Create file import dialog. This is a modal dialog so it will return only after user selection
+		const TSharedRef<IStormSyncImportWizard> Wizard = FStormSyncEditorModule::Get().CreateWizard(FilesToImport, AllFileInfos);
 
-			// Early out if user canceled the operation
-			if (!Wizard->ShouldImport())
-			{
-				UStormSyncImportSubsystem::Get().HandlePakPostExtract(PackageDescriptor, 0);
-				return;
-			}
+		// Early out if user canceled the operation
+		if (!Wizard->ShouldImport())
+		{
+			UStormSyncImportSubsystem::Get().HandlePakPostExtract(PackageDescriptor, 0);
+			return true;
 		}
+	}
 #endif
 
-		// Begin import process
-		{
-			TArray<const FStormSyncImportFileInfo*> ExistingFiles;
-			TArray<const FStormSyncImportFileInfo*> NewFiles;
-			SplitImportFiles(FilesToImport, ExistingFiles, NewFiles);
+	TArray<FText> ExtractErrors;
+	bool bExtractSuccess;
+	
+	// Begin import process
+	{
+		FScopedSlowTask SlowTaskExtract(FilesToImport.Num() + ExistingFilesToImport.Num() + NewFilesToImport.Num()
+			, LOCTEXT("ImportingFilesToProject", "Importing files to project..."));
+		SlowTaskExtract.MakeDialog();
 
-			FScopedSlowTask SlowTaskExtract(FilesToImport.Num() + ExistingFiles.Num() + NewFiles.Num()
-				, LOCTEXT("ImportingFilesToProject", "Importing files to project..."));
+		// Step 1: Prompt for Checkout (for existing assets)
+		SlowTaskExtract.EnterProgressFrame(ExistingFilesToImport.Num());
+		UStormSyncImportSubsystem::Get().HandleExistingAssetsPreExtract(ExistingFilesToImport, bShowWizard);
 
-			// Step 1: Prompt for Checkout (for existing assets)
-			SlowTaskExtract.EnterProgressFrame(ExistingFiles.Num());
-			UStormSyncImportSubsystem::Get().HandleExistingAssets(ExistingFiles, bShowWizard);
+		// Step 2: Extract
+		FStormSyncCoreExtractArgs ExtractArgs;
+		
+		ExtractArgs.OnPakPreExtract.BindLambda([&PackageDescriptor](const int32 FileCount)
+		{
+			UE_LOG(LogStormSyncImport, Display, TEXT("\tOnPakPreExtract - File Count: %d"), FileCount);
+			UStormSyncImportSubsystem::Get().HandlePakPreExtract(PackageDescriptor, FileCount);
+		});
+		
+		ExtractArgs.OnPakPostExtract.BindLambda([&PackageDescriptor, &FilesToImport](const int32 FileCount)
+		{
+			// Extract completion
+			UStormSyncImportSubsystem::Get().HandlePakPostExtract(PackageDescriptor, FilesToImport.Num());
+		});
 
-			// Step 2: Extract
-			for (const FStormSyncImportFileInfo& FileInfo : FilesToImport)
+		ExtractArgs.OnGetArchiveForExtract.BindLambda([&SlowTaskExtract, &FilesToImportMap](const FStormSyncFileDependency& InFileDependency, const FString& InDestFilepath) -> FStormSyncArchivePtr
+		{
+			if (const FStormSyncImportFileInfo** FoundFileInfo = FilesToImportMap.Find(InDestFilepath))
 			{
 				SlowTaskExtract.EnterProgressFrame();
-	
-				// Handle extraction, this is where assets are imported to project (closing asset editor if opened, deleting the asset and then writing it back)
-				UStormSyncImportSubsystem::Get().HandlePakAssetExtract(FileInfo.FileDependency, FileInfo.DestFilepath, FileInfo.FileBuffer);
-			}
 
-			// Step 3: Mark for Add (for new assets)
-			SlowTaskExtract.EnterProgressFrame(NewFiles.Num());
-			UStormSyncImportSubsystem::Get().HandleNewAssets(NewFiles, bShowWizard);
-		}
+				if (UStormSyncImportSubsystem::Get().HandlePakAssetPreExtract(**FoundFileInfo))
+				{
+					FStormSyncArchivePtr FileWriter(IFileManager::Get().CreateFileWriter(*InDestFilepath));
+					if (FileWriter)
+					{
+						UE_LOG(LogStormSyncImport, Verbose, TEXT("UStormSyncImportSubsystem::WriteFile - Creating file `%s`"), *InDestFilepath);
+					}
+					else
+					{
+						// failed report
+						UStormSyncImportSubsystem::Get().HandlePakAssetPostExtract(**FoundFileInfo, /*bFileWritten*/ false);
+					}
+					return FileWriter;
+				}
+			}
+			else
+			{
+				UE_LOG(LogStormSyncImport, Error, TEXT("UStormSyncImportSubsystem::PerformImport - FileInfo `%s` not found in import map."), *InDestFilepath);
+			}
+			
+			// If the files is not different, we don't need to import it nor load it. It will be skipped in the serializer. (not an error)
+			return nullptr;
+		});
 
-		// Extract completion
-		UStormSyncImportSubsystem::Get().HandlePakPostExtract(PackageDescriptor, FilesToImport.Num());
-	});
-	
-	ExtractArgs.OnFileExtract.BindLambda([&SlowTask, bDryRun, &FilesToImport, &BufferFiles](const FStormSyncFileDependency& FileDependency, const FString DestFilepath, const FStormSyncBufferPtr& FileBuffer)
-	{
-		SlowTask.EnterProgressFrame();
-		
-		if (bDryRun)
+		ExtractArgs.OnArchiveExtracted.BindLambda([&FilesToImportMap](const FStormSyncFileDependency& InFileDependency, const FString& InDestFilepath, const FStormSyncArchivePtr& InArchive)
 		{
-			UE_LOG(LogStormSyncImport, Display, TEXT("\tOnFileExtract"));
-			UE_LOG(LogStormSyncImport, Display, TEXT("\t\tPackageName: %s"), *FileDependency.PackageName.ToString());
-			UE_LOG(LogStormSyncImport, Display, TEXT("\t\tDestFilepath: %s"), *DestFilepath);
-			UE_LOG(LogStormSyncImport, Display, TEXT("\t\tFileSize: %lld"), FileDependency.FileSize);
-			UE_LOG(LogStormSyncImport, Display, TEXT("\t\tFileBuffer valid: %s"), FileBuffer.IsValid() ? TEXT("valid") : TEXT("invalid"));
-			if (FileBuffer.IsValid())
+			// Check that we wrote the expected amount of data.
+			const bool bFileWritten = InArchive && InArchive->Tell() == InFileDependency.FileSize;
+			if (const FStormSyncImportFileInfo** FoundFileInfo = FilesToImportMap.Find(InDestFilepath))
 			{
-				UE_LOG(LogStormSyncImport, Display, TEXT("\t\tFileBuffer size: %d"), FileBuffer->Num());
+				UStormSyncImportSubsystem::Get().HandlePakAssetPostExtract(**FoundFileInfo, bFileWritten);
 			}
-
-			return;
-		}
+			else
+			{
+				UE_LOG(LogStormSyncImport, Error, TEXT("UStormSyncImportSubsystem::PerformImport - FileInfo `%s` not found in import map."), *InDestFilepath);
+			}
+		});
 		
-		// Figure out files that needs import and build list of FStormSyncImportFileInfo for Import Wizard widget
-		const FStormSyncImportFileInfo FileInfo = FStormSyncImportFileInfo(FileDependency, DestFilepath, FileBuffer);
-		BufferFiles.Add(FileInfo);
+		InArchive->Seek(StartPosition);
+		bExtractSuccess = FStormSyncCoreUtils::ExtractPakBuffer(*InArchive.Get(), ExtractArgs, ExtractErrors);
 	
-		UStormSyncImportSubsystem::Get().FillFilesToImport(FileInfo, DestFilepath, FilesToImport);
-	});
-
-	TArray<FText> Errors;
-	TMap<FString, FString> ExtractedPackages;
-	if (!FStormSyncCoreUtils::ExtractPakBuffer(*InBuffer.Get(), ExtractArgs, ExtractedPackages, Errors))
+		// Step 3: Mark for Add (for new assets)
+		SlowTaskExtract.EnterProgressFrame(NewFilesToImport.Num());
+		UStormSyncImportSubsystem::Get().HandleNewAssetsPostExtract(NewFilesToImport, bShowWizard);
+	}
+	
+	if (!bExtractSuccess)
 	{
 		UE_LOG(LogStormSyncImport, Warning, TEXT("FStormSyncEditorUtils::PerformImport - Error extracting package ..."));
-		for (const FText& Error : Errors)
+		for (const FText& Error : ExtractErrors)
 		{
 			UE_LOG(LogStormSyncImport, Warning, TEXT("\t %s"), *Error.ToString());
 		}
-
-		return false;
 	}
 
-	return true;
+	return bExtractSuccess;
 }
 
 bool UStormSyncImportSubsystem::HandleNextTick(float InDeltaTime)
@@ -294,33 +395,29 @@ bool UStormSyncImportSubsystem::HandleNextTick(float InDeltaTime)
 	return false;
 }
 
-bool UStormSyncImportSubsystem::FillFilesToImport(const FStormSyncImportFileInfo& InFileInfo, const FString& InDestFilepath, TArray<FStormSyncImportFileInfo>& OutFilesToImport) const
+bool UStormSyncImportSubsystem::ShouldFileBeImported(FStormSyncImportFileInfo& InFileInfo)
 {
 	const FStormSyncFileDependency FileDependency = InFileInfo.FileDependency;
-
-	// Local copy
-	FStormSyncImportFileInfo FileInfo = InFileInfo;
 	
-	const int64 DestFileSize = IFileManager::Get().FileSize(*InDestFilepath);
+	const int64 DestFileSize = IFileManager::Get().FileSize(*InFileInfo.DestFilepath);
 	// File not existing in local project, add it to files to import and early out
 	if (DestFileSize == INDEX_NONE)
 	{
-		UE_LOG(LogStormSyncImport, Verbose, TEXT("\t\tFile %s does not exist"), *InDestFilepath);
-		FString ShortPath = InDestFilepath;
+		UE_LOG(LogStormSyncImport, Verbose, TEXT("\t\tFile %s does not exist"), *InFileInfo.DestFilepath);
+		FString ShortPath = InFileInfo.DestFilepath;
 		ShortPath.RemoveFromStart(FPaths::ProjectContentDir());
-		FileInfo.ImportReason = FText::Format(LOCTEXT("FileInfo_MissingFile", "Missing file in local project ({0})"), FText::FromString(ShortPath));
-		FileInfo.ImportReasonTooltip = FText::Format(LOCTEXT("FileInfo_MissingFile_Tooltip", "File {0} does not exist locally"), FText::FromString(InDestFilepath));
-		OutFilesToImport.Add(FileInfo);
+		InFileInfo.ImportReason = FText::Format(LOCTEXT("FileInfo_MissingFile", "Missing file in local project ({0})"), FText::FromString(ShortPath));
+		InFileInfo.ImportReasonTooltip = FText::Format(LOCTEXT("FileInfo_MissingFile_Tooltip", "File {0} does not exist locally"), FText::FromString(InFileInfo.DestFilepath));
 		return true;
 	}
 
-	UE_LOG(LogStormSyncImport, Verbose, TEXT("\t\tFile %s exist"), *InDestFilepath);
+	UE_LOG(LogStormSyncImport, Verbose, TEXT("\t\tFile %s exist"), *InFileInfo.DestFilepath);
 
 	const bool bSameSize = FileDependency.FileSize == DestFileSize;
 	UE_LOG(LogStormSyncImport, Verbose, TEXT("\t\t\tSame Size: %s (%lld vs %lld)"), bSameSize ? TEXT("true") : TEXT("false"), FileDependency.FileSize, DestFileSize);
 
 	// Check file hash from buffer against local file
-	const FMD5Hash ExistingFileMD5 = FMD5Hash::HashFile(*InDestFilepath);
+	const FMD5Hash ExistingFileMD5 = FMD5Hash::HashFile(*InFileInfo.DestFilepath);
 	const FString ExistingFileHash = LexToString(ExistingFileMD5);
 
 	const bool bSameHash = FileDependency.FileHash == ExistingFileHash;
@@ -328,31 +425,27 @@ bool UStormSyncImportSubsystem::FillFilesToImport(const FStormSyncImportFileInfo
 	// Mismatched File size or hash
 	if (!bSameSize)
 	{
-		FileInfo.ImportReason = FText::Format(
+		InFileInfo.ImportReason = FText::Format(
 			LOCTEXT("FileInfo_MismatchedSize", "Files have different sizes ({0} vs {1})"),
 			FText::FromString(FStormSyncCoreUtils::GetHumanReadableByteSize(DestFileSize)),
 			FText::FromString(FStormSyncCoreUtils::GetHumanReadableByteSize(FileDependency.FileSize))
 		);
-		FileInfo.ImportReasonTooltip = FText::Format(
+		InFileInfo.ImportReasonTooltip = FText::Format(
 			LOCTEXT("FileInfo_MismatchedSize_Tooltip", "Mismatched file size ({0} vs {1})"),
 			FText::AsNumber(DestFileSize),
 			FText::AsNumber(FileDependency.FileSize)
 		);
-
-		OutFilesToImport.Add(FileInfo);
 		return true;
 	}
 	
 	if (!bSameHash)
 	{
-		FileInfo.ImportReason = LOCTEXT("FileInfo_MismatchedHash", "Files hash are not matching.");
-		FileInfo.ImportReasonTooltip = FText::Format(
+		InFileInfo.ImportReason = LOCTEXT("FileInfo_MismatchedHash", "Files hash are not matching.");
+		InFileInfo.ImportReasonTooltip = FText::Format(
 			LOCTEXT("FileInfo_MismatchedHash_Tooltip", "Mismatched file hash ({0} vs {1})"),
 			FText::FromString(ExistingFileHash),
 			FText::FromString(FileDependency.FileHash)
 		);
-
-		OutFilesToImport.Add(FileInfo);
 		return true;
 	}
 
@@ -423,7 +516,7 @@ void UStormSyncImportSubsystem::HandlePakPostExtract(const FStormSyncPackageDesc
 #endif
 }
 
-void UStormSyncImportSubsystem::HandleExistingAssets(TConstArrayView<const FStormSyncImportFileInfo*> InExistingFiles, bool bInShowPrompt)
+void UStormSyncImportSubsystem::HandleExistingAssetsPreExtract(TConstArrayView<const FStormSyncImportFileInfo*> InExistingFiles, bool bInShowPrompt)
 {
 #if WITH_EDITOR
 	UE_LOG(LogStormSyncImport, Display, TEXT("UStormSyncImportSubsystem::HandleExistingAssets - Package Count: %d"), InExistingFiles.Num());
@@ -501,7 +594,7 @@ void UStormSyncImportSubsystem::HandleExistingAssets(TConstArrayView<const FStor
 #endif
 }
 
-bool UStormSyncImportSubsystem::HandleNewAssets(TConstArrayView<const FStormSyncImportFileInfo*> InNewFiles, bool bInShowPrompt)
+bool UStormSyncImportSubsystem::HandleNewAssetsPostExtract(TConstArrayView<const FStormSyncImportFileInfo*> InNewFiles, bool bInShowPrompt)
 {
 #if WITH_EDITOR
 	if (!GEditor || GIsAutomationTesting || IsRunningCommandlet())
@@ -545,52 +638,23 @@ bool UStormSyncImportSubsystem::HandleNewAssets(TConstArrayView<const FStormSync
 #endif
 }
 
-void UStormSyncImportSubsystem::HandlePakAssetExtract(const FStormSyncFileDependency& FileDependency, const FString& DestFilepath, const FStormSyncBufferPtr& FileBuffer)
+bool UStormSyncImportSubsystem::HandlePakAssetPreExtract(const FStormSyncImportFileInfo& InFileInfo)
 {
-	UE_LOG(LogStormSyncImport, Display, TEXT("UStormSyncImportSubsystem::HandlePakAssetExtract - Handle extracted package: %s"), *FileDependency.ToString());
+	const FStormSyncFileDependency& FileDependency = InFileInfo.FileDependency;
+	const FString& DestFilepath = InFileInfo.DestFilepath;
+	
+	UE_LOG(LogStormSyncImport, Display, TEXT("UStormSyncImportSubsystem::HandlePakAssetPreExtract - Handle extracted package: %s"), *FileDependency.ToString());
 	UE_LOG(LogStormSyncImport, Display, TEXT("\tshould extract to %s"), *DestFilepath);
-	UE_LOG(LogStormSyncImport, Display, TEXT("\tValid buffer ? %s"), FileBuffer.IsValid() && FileBuffer->Num() != 0 ? TEXT("true") : TEXT("false"));
-
-	if (!FileBuffer.IsValid())
-	{
-		UE_LOG(LogStormSyncImport, Error, TEXT("UStormSyncImportSubsystem::HandlePakAssetExtract - Invalid file buffer"));
-		return;
-	}
-
 
 	FStormSyncEditorFileReport FileReport;
 	FileReport.PackageName = FileDependency.PackageName;
 	FileReport.DestFilepath = DestFilepath;
 
-	if (!FPaths::FileExists(DestFilepath))
+	if (InFileInfo.bNewAsset)
 	{
-		if (WriteFile(DestFilepath, FileDependency.FileSize, FileBuffer->GetData()))
-		{
-			FileReport.bSuccess = true;
-		}
-		else
-		{
-			FileReport.bSuccess = false;
-		}
-
-		ExtractedFileReports.Add(FileReport);
-
-#if WITH_EDITOR
-		FMessageLog MessageLog(LogName);
-		if (FileReport.bSuccess)
-		{
-			MessageLog.Info(FText::Format(LOCTEXT("PakAssetExtract_SucessfulExtract", "Extracted {0} file to {1}."), FText::FromName(FileDependency.PackageName), FText::FromString(DestFilepath)));
-		}
-		else
-		{
-			MessageLog.Error(FText::Format(LOCTEXT("PakAssetExtract_FileWrite_Failed", "Failed to extract {0} file to {1}"), FText::FromName(FileDependency.PackageName), FText::FromString(DestFilepath)));
-		}
-#endif
-
-		return;
+		return true;	// Nothing to do.
 	}
-
-
+	
 	// If File exists locally, we should:
 	//
 	// 1. Close Editors
@@ -621,9 +685,10 @@ void UStormSyncImportSubsystem::HandlePakAssetExtract(const FStormSyncFileDepend
 	{
 		if (!DeleteAssets(AssetsToDelete))
 		{
+			UE_LOG(LogStormSyncImport, Error, TEXT("UStormSyncImportSubsystem::HandlePakAssetPreExtract - Unable to unload asset for package: %s"), *FileDependency.PackageName.ToString());
 			FileReport.bSuccess = false;
 			ExtractedFileReports.Add(FileReport);
-			return;
+			return false;
 		}
 	}
 	else
@@ -632,14 +697,26 @@ void UStormSyncImportSubsystem::HandlePakAssetExtract(const FStormSyncFileDepend
 		FlushPackageLoading(FileDependency.PackageName.ToString());
 	}
 
-	// Write the file to destination
-	if (WriteFile(DestFilepath, FileDependency.FileSize, FileBuffer->GetData()))
+	return true;
+}
+
+void UStormSyncImportSubsystem::HandlePakAssetPostExtract(const FStormSyncImportFileInfo& InFileInfo, bool bInFileWritten)
+{
+	const FStormSyncFileDependency& FileDependency = InFileInfo.FileDependency;
+	const FString& DestFilepath = InFileInfo.DestFilepath;
+
+	FStormSyncEditorFileReport FileReport;
+	FileReport.PackageName = FileDependency.PackageName;
+	FileReport.DestFilepath = DestFilepath;
+	FileReport.bSuccess = bInFileWritten;
+
+	if (FileReport.bSuccess)
 	{
-		FileReport.bSuccess = true;
+		UE_LOG(LogStormSyncImport, Display, TEXT("UStormSyncImportSubsystem::HandlePakAssetPostExtract - Package: \"%s\" was extracted to \"%s\""), *FileDependency.ToString(), *DestFilepath);
 	}
 	else
 	{
-		FileReport.bSuccess = false;
+		UE_LOG(LogStormSyncImport, Error, TEXT("UStormSyncImportSubsystem::HandlePakAssetPostExtract - Package: \"%s\" failed to extract to \"%s\""), *FileDependency.ToString(), *DestFilepath);
 	}
 
 #if WITH_EDITOR
@@ -654,7 +731,7 @@ void UStormSyncImportSubsystem::HandlePakAssetExtract(const FStormSyncFileDepend
 	}
 #endif
 
-	ExtractedFileReports.Add(FileReport);
+	ExtractedFileReports.Add(MoveTemp(FileReport));
 }
 
 void UStormSyncImportSubsystem::CloseEditors(const TArray<FAssetData>& InAssets, TArray<FAssetData>& OutClosedAssets)
@@ -949,28 +1026,4 @@ void UStormSyncImportSubsystem::HotReloadPackages(const TArray<FStormSyncEditorF
 	}
 }
 
-void UStormSyncImportSubsystem::SplitImportFiles(TConstArrayView<FStormSyncImportFileInfo> InFiles, TArray<const FStormSyncImportFileInfo*>& OutExistingFiles, TArray<const FStormSyncImportFileInfo*>& OutNewFiles)
-{
-	if (InFiles.IsEmpty())
-	{
-		return;
-	}
-
-	// Worst case applied to both
-	OutExistingFiles.Empty(InFiles.Num());
-	OutNewFiles.Empty(InFiles.Num());
-
-	for (const FStormSyncImportFileInfo& FileInfo : InFiles)
-	{
-		if (FPaths::FileExists(FileInfo.DestFilepath))
-		{
-			OutExistingFiles.Add(&FileInfo);
-		}
-		else
-		{
-			OutNewFiles.Add(&FileInfo);
-		}
-	}
-}
-
 #undef LOCTEXT_NAMESPACE
diff --git a/Engine/Plugins/Experimental/StormSync/Source/StormSyncImport/Private/Subsystems/StormSyncImportSubsystem.h b/Engine/Plugins/Experimental/StormSync/Source/StormSyncImport/Private/Subsystems/StormSyncImportSubsystem.h
--- a/Engine/Plugins/Experimental/StormSync/Source/StormSyncImport/Private/Subsystems/StormSyncImportSubsystem.h
+++ b/Engine/Plugins/Experimental/StormSync/Source/StormSyncImport/Private/Subsystems/StormSyncImportSubsystem.h
@@ -81,7 +81,7 @@ public:
 	 * 
 	 * @return Operation status, true if successful
 	 */
-	static bool PerformBufferImport(const FStormSyncPackageDescriptor& InPackageDescriptor, const FStormSyncBufferPtr& InBuffer);
+	static bool PerformBufferImport(const FStormSyncPackageDescriptor& InPackageDescriptor, const FStormSyncArchivePtr& InArchive);
 
 	/**
 	 * Takes a raw buffer (a storm sync pak) and begins extraction process
@@ -90,14 +90,23 @@ public:
 	 * the list of files that needs to be extracted and updated in this editor instance.
 	 *
 	 * @param InPackageDescriptor The package descriptor for this buffer extraction. Only used for UI message logs.
-	 * @param InBuffer Raw buffer we'd like to extract files from to local project
+	 * @param InArchive Source archive to be parse to extract content
 	 * @param bShowWizard Whether or not to show the import wizard
 	 * @param bDryRun Whether we should proceed to extraction or simply log to console buffer content
 	 * 
 	 * @return Operation status, true if successful
 	 */
-	static bool PerformImport(const FStormSyncPackageDescriptor& InPackageDescriptor, const FStormSyncBufferPtr& InBuffer, bool bShowWizard = false, bool bDryRun = false);
+	static bool PerformImport(const FStormSyncPackageDescriptor& InPackageDescriptor, const FStormSyncArchivePtr& InArchive, bool bShowWizard = false, bool bDryRun = false);
 
+private:
+	/**
+	 * Parse the buffer to extract all the file information (but skip content)
+	 * @param InArchive Source archive of the pak
+	 * @param OutFileInfos File info in the pak.
+	 * @return true if parsing was successful.
+	 */
+	static bool ParsePak(const FStormSyncArchivePtr& InArchive, TArray<FStormSyncImportFileInfo>& OutFileInfos);
+	
 private:
 	/** The name of the message output log we will send messages to */
 	static constexpr const TCHAR* LogName = TEXT("StormSyncImport");
@@ -122,7 +131,7 @@ private:
 	 * If a change is detected (either missing locally or mismatched file size / hash), OutFilesToImport is filled with a new file info indicating that the
 	 * file needs an import.
 	 */
-	bool FillFilesToImport(const FStormSyncImportFileInfo& InFileInfo, const FString& InDestFilepath, TArray<FStormSyncImportFileInfo>& OutFilesToImport) const;
+	static bool ShouldFileBeImported(FStormSyncImportFileInfo& InFileInfo);
 
 	/** Event handler for core delegate when a pak extraction starts */
 	void HandlePakPreExtract(const FStormSyncPackageDescriptor& InPackageDescriptor, int32 FileCount);
@@ -131,13 +140,16 @@ private:
 	void HandlePakPostExtract(const FStormSyncPackageDescriptor& InPackageDescriptor, int32 FileCount) const;
 
 	/** Event handler for existing asset checkout prior to extracting the new overwriting assets */
-	void HandleExistingAssets(TConstArrayView<const FStormSyncImportFileInfo*> InExistingFiles, bool bInShowPrompt);
+	void HandleExistingAssetsPreExtract(TConstArrayView<const FStormSyncImportFileInfo*> InExistingFiles, bool bInShowPrompt);
 
 	/** Event handler for auto-adding new assets post extracting them */
-	bool HandleNewAssets(TConstArrayView<const FStormSyncImportFileInfo*> InNewFiles, bool bInShowPrompt);
+	bool HandleNewAssetsPostExtract(TConstArrayView<const FStormSyncImportFileInfo*> InNewFiles, bool bInShowPrompt);
+
+	/** Event handler for when a file is about to be extracted from an incoming pak */
+	bool HandlePakAssetPreExtract(const FStormSyncImportFileInfo& InFileInfo);
 
-	/** Event handler for core delegate when a file is extracted from an incoming pak */
-	void HandlePakAssetExtract(const FStormSyncFileDependency& FileDependency, const FString& DestFilepath, const FStormSyncBufferPtr& FileBuffer);
+	/** Event handler for when a file has been extracted from an incoming pak */
+	void HandlePakAssetPostExtract(const FStormSyncImportFileInfo& InFileInfo, bool bInFileWritten);
 
 	/** Checks if given asset is currently being edited, and if so, closes the editor (required prior to file deletion) */
 	static void CloseEditors(const TArray<FAssetData>& InAssets, TArray<FAssetData>& OutClosedAssets);
@@ -169,7 +181,4 @@ private:
 
 	/** Attempts to reload the specified top-level packages. */
 	static void HotReloadPackages(const TArray<FStormSyncEditorFileReport>& InExtractedFileReports, bool bInInteractiveHotReload);
-
-	/** Generates two arrays with pointers to the input file array. One with the already-existing files, and the other with the new files that would be generated in import */
-	static void SplitImportFiles(TConstArrayView<FStormSyncImportFileInfo> InFiles, TArray<const FStormSyncImportFileInfo*>& OutExistingFiles, TArray<const FStormSyncImportFileInfo*>& OutNewFiles);
 };
diff --git a/Engine/Plugins/Experimental/StormSync/Source/StormSyncImport/Private/Subsystems/StormSyncImportWorldSubsystem.cpp b/Engine/Plugins/Experimental/StormSync/Source/StormSyncImport/Private/Subsystems/StormSyncImportWorldSubsystem.cpp
--- a/Engine/Plugins/Experimental/StormSync/Source/StormSyncImport/Private/Subsystems/StormSyncImportWorldSubsystem.cpp
+++ b/Engine/Plugins/Experimental/StormSync/Source/StormSyncImport/Private/Subsystems/StormSyncImportWorldSubsystem.cpp
@@ -41,19 +41,19 @@ void UStormSyncImportWorldSubsystem::HandleImportFile(const FString& InFilename)
 	UStormSyncImportSubsystem::Get().EnqueueImportTask(MakeShared<FStormSyncImportFilesTask>(InFilename), GetWorld());
 }
 
-void UStormSyncImportWorldSubsystem::HandleImportBuffer(const FStormSyncPackageDescriptor& InPackageDescriptor, const FStormSyncBufferPtr& InBuffer) const
+void UStormSyncImportWorldSubsystem::HandleImportBuffer(const FStormSyncPackageDescriptor& InPackageDescriptor, const FStormSyncArchivePtr& InArchive) const
 {
-	if (!InBuffer.IsValid())
+	if (!InArchive.IsValid())
 	{
-		UE_LOG(LogStormSyncImport, Error, TEXT("UStormSyncImportSubsystem::HandleImportBuffer failed from invalid buffer"));
+		UE_LOG(LogStormSyncImport, Error, TEXT("UStormSyncImportSubsystem::HandleImportBuffer failed from invalid archive"));
 		return;
 	}
 	
 	UE_LOG(LogStormSyncImport, 
 		Verbose,
-		TEXT("UStormSyncImportWorldSubsystem::HandleImportBuffer InPackageDescriptor: %s, BufferSize: %d (World: %s)"),
+		TEXT("UStormSyncImportWorldSubsystem::HandleImportBuffer InPackageDescriptor: %s, BufferSize: %lld (World: %s)"),
 		*InPackageDescriptor.ToString(),
-		InBuffer->Num(),
+		InArchive->TotalSize(),
 		*GetNameSafe(GetWorld())
 	);
 
@@ -64,5 +64,5 @@ void UStormSyncImportWorldSubsystem::HandleImportBuffer(const FStormSyncPackageD
 		return;
 	}
 
-	UStormSyncImportSubsystem::Get().EnqueueImportTask(MakeShared<FStormSyncImportBufferTask>(InPackageDescriptor, InBuffer), GetWorld());
+	UStormSyncImportSubsystem::Get().EnqueueImportTask(MakeShared<FStormSyncImportBufferTask>(InPackageDescriptor, InArchive), GetWorld());
 }
diff --git a/Engine/Plugins/Experimental/StormSync/Source/StormSyncImport/Private/Subsystems/StormSyncImportWorldSubsystem.h b/Engine/Plugins/Experimental/StormSync/Source/StormSyncImport/Private/Subsystems/StormSyncImportWorldSubsystem.h
--- a/Engine/Plugins/Experimental/StormSync/Source/StormSyncImport/Private/Subsystems/StormSyncImportWorldSubsystem.h
+++ b/Engine/Plugins/Experimental/StormSync/Source/StormSyncImport/Private/Subsystems/StormSyncImportWorldSubsystem.h
@@ -26,5 +26,5 @@ public:
 	virtual void Deinitialize() override;
 	
 	void HandleImportFile(const FString& InFilename) const;
-	void HandleImportBuffer(const FStormSyncPackageDescriptor& InPackageDescriptor, const FStormSyncBufferPtr& InBuffer) const;
+	void HandleImportBuffer(const FStormSyncPackageDescriptor& InPackageDescriptor, const FStormSyncArchivePtr& InArchive) const;
 };
diff --git a/Engine/Plugins/Experimental/StormSync/Source/StormSyncImport/Private/Tasks/StormSyncImportBufferTask.cpp b/Engine/Plugins/Experimental/StormSync/Source/StormSyncImport/Private/Tasks/StormSyncImportBufferTask.cpp
--- a/Engine/Plugins/Experimental/StormSync/Source/StormSyncImport/Private/Tasks/StormSyncImportBufferTask.cpp
+++ b/Engine/Plugins/Experimental/StormSync/Source/StormSyncImport/Private/Tasks/StormSyncImportBufferTask.cpp
@@ -7,12 +7,12 @@
 
 void FStormSyncImportBufferTask::Run()
 {
-	if (!Buffer.IsValid())
+	if (!Archive.IsValid())
 	{
-		UE_LOG(LogStormSyncImport, Error, TEXT("FStormSyncImportBufferTask::Run failed on invalid buffer"));
+		UE_LOG(LogStormSyncImport, Error, TEXT("FStormSyncImportBufferTask::Run failed on invalid archive"));
 		return;
 	}
 
-	UE_LOG(LogStormSyncImport, Display, TEXT("FStormSyncImportBufferTask::Run for buffer of size %d"), Buffer->Num());
-	UStormSyncImportSubsystem::Get().PerformBufferImport(PackageDescriptor, MoveTemp(Buffer));
+	UE_LOG(LogStormSyncImport, Display, TEXT("FStormSyncImportBufferTask::Run for buffer of size %lld"), Archive->TotalSize());
+	UStormSyncImportSubsystem::Get().PerformBufferImport(PackageDescriptor, MoveTemp(Archive));
 }
diff --git a/Engine/Plugins/Experimental/StormSync/Source/StormSyncImport/Private/Tasks/StormSyncImportBufferTask.h b/Engine/Plugins/Experimental/StormSync/Source/StormSyncImport/Private/Tasks/StormSyncImportBufferTask.h
--- a/Engine/Plugins/Experimental/StormSync/Source/StormSyncImport/Private/Tasks/StormSyncImportBufferTask.h
+++ b/Engine/Plugins/Experimental/StormSync/Source/StormSyncImport/Private/Tasks/StormSyncImportBufferTask.h
@@ -2,6 +2,7 @@
 
 #pragma once
 
+#include "StormSyncCommonTypes.h"
 #include "StormSyncPackageDescriptor.h"
 #include "Tasks/IStormSyncImportTask.h"
 
@@ -9,9 +10,9 @@
 class FStormSyncImportBufferTask final : public IStormSyncImportSubsystemTask
 {
 public:
-	explicit FStormSyncImportBufferTask(const FStormSyncPackageDescriptor& InPackageDescriptor, const TSharedPtr<TArray<uint8>, ESPMode::ThreadSafe>& InBuffer)
+	explicit FStormSyncImportBufferTask(const FStormSyncPackageDescriptor& InPackageDescriptor, const FStormSyncArchivePtr& InArchive)
 		: PackageDescriptor(InPackageDescriptor)
-		, Buffer(InBuffer)
+		, Archive(InArchive)
 	{
 	}
 
@@ -21,6 +22,6 @@ private:
 	/** Metadata info about buffer being extracted */
 	FStormSyncPackageDescriptor PackageDescriptor;
 	
-	/** Holds a shared ptr to the serialized buffer */
-	TSharedPtr<TArray<uint8>, ESPMode::ThreadSafe> Buffer = nullptr;
+	/** Archive to import */
+	FStormSyncArchivePtr Archive;
 };
\ No newline at end of file
diff --git a/Engine/Plugins/Experimental/StormSync/Source/StormSyncTests/Private/StormSyncCore.spec.cpp b/Engine/Plugins/Experimental/StormSync/Source/StormSyncTests/Private/StormSyncCore.spec.cpp
--- a/Engine/Plugins/Experimental/StormSync/Source/StormSyncTests/Private/StormSyncCore.spec.cpp
+++ b/Engine/Plugins/Experimental/StormSync/Source/StormSyncTests/Private/StormSyncCore.spec.cpp
@@ -4,6 +4,8 @@
 #include "HAL/FileManager.h"
 #include "Misc/AutomationTest.h"
 #include "Misc/PackageName.h"
+#include "Serialization/MemoryReader.h"
+#include "Serialization/MemoryWriter.h"
 #include "StormSyncCoreSettings.h"
 #include "StormSyncCoreUtils.h"
 
@@ -325,14 +327,15 @@ void FStormSyncCoreSpec::Define()
 		It(TEXT("should fail if entries input is empty"), [this]()
 		{
 			FText ErrorText;
-			TArray<uint8> PakBuffer;
+			FStormSyncBuffer PakBuffer;
+			FMemoryWriter64 PakArchiveWriter(PakBuffer);
 			const TArray<FName> PackageNames;
 
-			const bool bResult = FStormSyncCoreUtils::CreatePakBufferWithDependencies(PackageNames, PakBuffer, ErrorText);
+			const bool bResult = FStormSyncCoreUtils::CreatePakBufferWithDependencies(PackageNames, PakArchiveWriter, ErrorText);
 
 			TestFalse(TEXT("CreatePakBufferWithDependencies result is false"), bResult);
 			TestEqual(TEXT("CreatePakBufferWithDependencies ErrorText is the expected one"), ErrorText.ToString(), TEXT("Provided PackageNames array is empty."));
-			TestEqual(TEXT("CreatePakBufferWithDependencies PakBuffer remains empty"), PakBuffer.Num(), 0);
+			TestEqual(TEXT("CreatePakBufferWithDependencies PakBuffer remains empty"), PakBuffer.Num(), static_cast<int64>(0));
 		});
 
 		TArray<FName> FixturePackageNames;
@@ -344,7 +347,8 @@ void FStormSyncCoreSpec::Define()
 			{
 				FText ErrorText;
 
-				TArray<uint8> PakBuffer;
+				FStormSyncBuffer PakBuffer;
+				FMemoryWriter64 PakArchiveWriter(PakBuffer);
 
 				TArray<FName> PackageNames;
 				PackageNames.Add(FixturePackageName);
@@ -358,7 +362,7 @@ void FStormSyncCoreSpec::Define()
 					PackageDescriptor.Dependencies.Add(FileDependency);
 				});
 
-				const bool bResult = FStormSyncCoreUtils::CreatePakBufferWithDependencies(PackageNames, PakBuffer, ErrorText, Delegate);
+				const bool bResult = FStormSyncCoreUtils::CreatePakBufferWithDependencies(PackageNames, PakArchiveWriter, ErrorText, Delegate);
 
 				TestTrue(TEXT("CreatePakBufferWithDependencies result is true"), bResult);
 				TestTrue(TEXT("ErrorText is empty"), ErrorText.IsEmpty());
@@ -386,8 +390,9 @@ void FStormSyncCoreSpec::Define()
 			FText ErrorText;
 			const TArray<FName> PackageNames;
 
-			TArray<uint8> PakBuffer;
-			const bool bResult = FStormSyncCoreUtils::CreatePakBuffer(PackageNames, PakBuffer, ErrorText);
+			FStormSyncBuffer PakBuffer;
+			FMemoryWriter64 PakArchiveWriter(PakBuffer);
+			const bool bResult = FStormSyncCoreUtils::CreatePakBuffer(PackageNames, PakArchiveWriter, ErrorText);
 
 			TestFalse(TEXT("CreatePakBuffer result is false"), bResult);
 			TestEqual(TEXT("CreatePakBuffer ErrorText is the expected one"), ErrorText.ToString(), TEXT("Provided PackageNames array is empty."));
@@ -414,7 +419,8 @@ void FStormSyncCoreSpec::Define()
 					return;
 				}
 
-				TArray<uint8> PakBuffer;
+				FStormSyncBuffer PakBuffer;
+				FMemoryWriter64 PakArchiveWriter(PakBuffer);
 
 				TArray<FStormSyncFileDependency> SuccessfullyPackedFiles;
 				const FStormSyncCoreUtils::FOnFileAdded Delegate = FStormSyncCoreUtils::FOnFileAdded::CreateLambda([&SuccessfullyPackedFiles](const FStormSyncFileDependency& FileDependency)
@@ -422,7 +428,7 @@ void FStormSyncCoreSpec::Define()
 					SuccessfullyPackedFiles.Add(FileDependency);
 				});
 
-				const bool bResult = FStormSyncCoreUtils::CreatePakBuffer(PackageDependencies, PakBuffer, ErrorText, Delegate);
+				const bool bResult = FStormSyncCoreUtils::CreatePakBuffer(PackageDependencies, PakArchiveWriter, ErrorText, Delegate);
 
 				TestTrue(TEXT("CreatePakBufferWithDependencies result is true"), bResult);
 				TestTrue(TEXT("ErrorText is empty"), ErrorText.IsEmpty());
@@ -447,7 +453,8 @@ void FStormSyncCoreSpec::Define()
 			It(FString::Printf(TEXT("should extract package descriptor (%s)"), *FixturePackageName.ToString()), [this, FixturePackageName]()
 			{
 				FText ErrorText;
-				TArray<uint8> PakBuffer;
+				FStormSyncBuffer PakBuffer;
+				FMemoryWriter64 PakArchiveWriter(PakBuffer);
 
 				TArray<FName> PackageNames;
 				PackageNames.Add(FixturePackageName);
@@ -460,7 +467,7 @@ void FStormSyncCoreSpec::Define()
 					PackageDescriptor.Dependencies.Add(FileDependency);
 				});
 
-				if (!FStormSyncCoreUtils::CreatePakBufferWithDependencies(PackageNames, PakBuffer, ErrorText, Delegate))
+				if (!FStormSyncCoreUtils::CreatePakBufferWithDependencies(PackageNames, PakArchiveWriter, ErrorText, Delegate))
 				{
 					AddError(FString::Printf(TEXT("CreatePakBufferWithDependencies Error: %s"), *ErrorText.ToString()));
 					return;
@@ -479,7 +486,7 @@ void FStormSyncCoreSpec::Define()
 					AddInfo(FString::Printf(TEXT("\tOnPakPostExtract %d"), FileCount));
 				});
 
-				ExtractArgs.OnFileExtract.BindLambda([this](const FStormSyncFileDependency& FileDependency, const FString DestFilepath, const FStormSyncBufferPtr& FileBuffer)
+				ExtractArgs.OnFileExtract.BindLambda([this, &SuccessfullyExtractedPackages](const FStormSyncFileDependency& FileDependency, const FString DestFilepath, const FStormSyncBufferPtr& FileBuffer)
 				{
 					AddInfo(FString::Printf(TEXT("\tOnFileExtract")));
 					AddInfo(FString::Printf(TEXT("\t\tPackageName: %s"), *FileDependency.PackageName.ToString()));
@@ -488,11 +495,14 @@ void FStormSyncCoreSpec::Define()
 					AddInfo(FString::Printf(TEXT("\t\tFileBuffer valid: %s"), FileBuffer.IsValid() ? TEXT("valid") : TEXT("invalid")));
 					if (FileBuffer.IsValid())
 					{
-						AddInfo(FString::Printf(TEXT("\t\tFileBuffer size: %d"), FileBuffer->Num()));
+						AddInfo(FString::Printf(TEXT("\t\tFileBuffer size: %lld"), FileBuffer->Num()));
 					}
+
+					SuccessfullyExtractedPackages.Add(FileDependency.PackageName.ToString(), DestFilepath);
 				});
 
-				const bool bResult = FStormSyncCoreUtils::ExtractPakBuffer(PakBuffer, ExtractArgs, SuccessfullyExtractedPackages, Errors);
+				FMemoryReaderView PakArchiveReader(PakBuffer);
+				const bool bResult = FStormSyncCoreUtils::ExtractPakBuffer(PakArchiveReader, ExtractArgs, Errors);
 
 				TestTrue(TEXT("ExtractPakBuffer result is true"), bResult);
 				TestEqual(TEXT("ExtractPakBuffer SuccessfullyExtractedPackages number is equal to the original file numbers"), SuccessfullyExtractedPackages.Num(), PackageDescriptor.Dependencies.Num());
diff --git a/Engine/Plugins/Experimental/StormSync/Source/StormSyncTransportClient/Private/Socket/StormSyncTransportClientSocket.cpp b/Engine/Plugins/Experimental/StormSync/Source/StormSyncTransportClient/Private/Socket/StormSyncTransportClientSocket.cpp
--- a/Engine/Plugins/Experimental/StormSync/Source/StormSyncTransportClient/Private/Socket/StormSyncTransportClientSocket.cpp
+++ b/Engine/Plugins/Experimental/StormSync/Source/StormSyncTransportClient/Private/Socket/StormSyncTransportClientSocket.cpp
@@ -389,7 +389,7 @@ bool FStormSyncTransportClientSocket::IsConnected() const
 	return Socket && Socket->GetConnectionState() == SCS_Connected;
 }
 
-void FStormSyncTransportClientSocket::SendBuffer(const TArray<uint8>& InBuffer)
+void FStormSyncTransportClientSocket::SendBuffer(const FStormSyncBuffer& InBuffer)
 {
 	const int32 BufferSize = InBuffer.Num();
 
diff --git a/Engine/Plugins/Experimental/StormSync/Source/StormSyncTransportClient/Private/Socket/StormSyncTransportClientSocket.h b/Engine/Plugins/Experimental/StormSync/Source/StormSyncTransportClient/Private/Socket/StormSyncTransportClientSocket.h
--- a/Engine/Plugins/Experimental/StormSync/Source/StormSyncTransportClient/Private/Socket/StormSyncTransportClientSocket.h
+++ b/Engine/Plugins/Experimental/StormSync/Source/StormSyncTransportClient/Private/Socket/StormSyncTransportClientSocket.h
@@ -4,6 +4,7 @@
 
 #include "HAL/Runnable.h"
 #include "Interfaces/IPv4/IPv4Endpoint.h"
+#include "StormSyncCommonTypes.h"
 #include "StormSyncTransportSettings.h"
 
 /**
@@ -95,7 +96,7 @@ public:
 	 *
 	 * @param InBuffer The raw buffer to send
 	 */
-	void SendBuffer(const TArray<uint8>& InBuffer);
+	void SendBuffer(const FStormSyncBuffer& InBuffer);
 
 	/**
 	 * Gets the current state of the connection.
diff --git a/Engine/Plugins/Experimental/StormSync/Source/StormSyncTransportClient/Private/StormSyncTransportClientEndpoint.cpp b/Engine/Plugins/Experimental/StormSync/Source/StormSyncTransportClient/Private/StormSyncTransportClientEndpoint.cpp
--- a/Engine/Plugins/Experimental/StormSync/Source/StormSyncTransportClient/Private/StormSyncTransportClientEndpoint.cpp
+++ b/Engine/Plugins/Experimental/StormSync/Source/StormSyncTransportClient/Private/StormSyncTransportClientEndpoint.cpp
@@ -7,6 +7,7 @@
 #include "IMessagingModule.h"
 #include "MessageEndpointBuilder.h"
 #include "Misc/PackageName.h"
+#include "Serialization/MemoryWriter.h"
 #include "Services/StormSyncPullMessageService.h"
 #include "Services/StormSyncPushMessageService.h"
 #include "Services/StormSyncStatusMessageService.h"
@@ -222,8 +223,9 @@ void FStormSyncTransportClientEndpoint::StartSendingBuffer(const FStormSyncTrans
 	Async(EAsyncExecution::Thread, [LocalWeakThis, ServerAddress, PackageNames, InDoneDelegate]()
 	{
 		FText ErrorText;
-		TArray<uint8> PakBuffer;
-		if (!FStormSyncCoreUtils::CreatePakBuffer(PackageNames, PakBuffer, ErrorText))
+		FStormSyncBuffer PakBuffer;
+		FMemoryWriter64 PakArchiveWriter(PakBuffer);
+		if (!FStormSyncCoreUtils::CreatePakBuffer(PackageNames, PakArchiveWriter, ErrorText))
 		{
 			UE_LOG(LogStormSyncClient, Display, TEXT("FStormSyncTransportClientEndpoint::StartSendingBuffer - CreatePakBuffer Error: %s"), *ErrorText.ToString());
 			InDoneDelegate.ExecuteIfBound(MakeShared<FStormSyncSendingBufferPayload>(ErrorText));
@@ -330,7 +332,7 @@ TArray<FName> FStormSyncTransportClientEndpoint::GetPackageNamesFromModifierInfo
 	return PackageNames;
 }
 
-void FStormSyncTransportClientEndpoint::SendTcpBuffer(const FString& InAddress, const TArray<uint8>& InBuffer)
+void FStormSyncTransportClientEndpoint::SendTcpBuffer(const FString& InAddress, const FStormSyncBuffer& InBuffer)
 {
 	const TSharedPtr<FStormSyncTransportClientSocket> Connection = GetOrCreateClientSocket(InAddress);
 	if (!Connection)
diff --git a/Engine/Plugins/Experimental/StormSync/Source/StormSyncTransportClient/Private/StormSyncTransportClientEndpoint.h b/Engine/Plugins/Experimental/StormSync/Source/StormSyncTransportClient/Private/StormSyncTransportClientEndpoint.h
--- a/Engine/Plugins/Experimental/StormSync/Source/StormSyncTransportClient/Private/StormSyncTransportClientEndpoint.h
+++ b/Engine/Plugins/Experimental/StormSync/Source/StormSyncTransportClient/Private/StormSyncTransportClientEndpoint.h
@@ -132,7 +132,7 @@ private:
 	void HandlePongMessage(const FStormSyncTransportPongMessage& InMessage, const TSharedRef<IMessageContext, ESPMode::ThreadSafe>& InMessageContext);
 	
 	/** Forwards sending of raw buffer over a tcp connection */
-	void SendTcpBuffer(const FString& InAddress, const TArray<uint8>& InBuffer);
+	void SendTcpBuffer(const FString& InAddress, const FStormSyncBuffer& InBuffer);
 
 	/** Handler for when a tcp socket is closing */
 	void OnConnectionClosedForSocket(const FIPv4Endpoint& InEndpoint);
diff --git a/Engine/Plugins/Experimental/StormSync/Source/StormSyncTransportClient/Private/Tests/StormSyncTransportClientEndpoint.spec.cpp b/Engine/Plugins/Experimental/StormSync/Source/StormSyncTransportClient/Private/Tests/StormSyncTransportClientEndpoint.spec.cpp
--- a/Engine/Plugins/Experimental/StormSync/Source/StormSyncTransportClient/Private/Tests/StormSyncTransportClientEndpoint.spec.cpp
+++ b/Engine/Plugins/Experimental/StormSync/Source/StormSyncTransportClient/Private/Tests/StormSyncTransportClientEndpoint.spec.cpp
@@ -137,7 +137,7 @@ public:
 							SendMessage(CreateMessage(FStormSyncTransportTcpTransferCompletePacket()));
 							
 							UE_LOG(LogTemp, Display, TEXT("\tFMockTcpServer: We received full buffer!!"));
-							OnReceivedBuffer.ExecuteIfBound(MakeShared<TArray<uint8>>(MessageBuffer));
+							OnReceivedBuffer.ExecuteIfBound(MakeShared<FStormSyncBuffer>(MessageBuffer));
 						}
 					}
 				}
diff --git a/Engine/Plugins/Experimental/StormSync/Source/StormSyncTransportServer/Private/Socket/StormSyncTransportTcpServer.cpp b/Engine/Plugins/Experimental/StormSync/Source/StormSyncTransportServer/Private/Socket/StormSyncTransportTcpServer.cpp
--- a/Engine/Plugins/Experimental/StormSync/Source/StormSyncTransportServer/Private/Socket/StormSyncTransportTcpServer.cpp
+++ b/Engine/Plugins/Experimental/StormSync/Source/StormSyncTransportServer/Private/Socket/StormSyncTransportTcpServer.cpp
@@ -268,7 +268,7 @@ void FStormSyncTransportTcpServer::HandleIncomingBuffer(const FIPv4Endpoint& InE
 void FStormSyncTransportTcpServer::HandleReceivedBuffer(const FIPv4Endpoint& InEndpoint, const TSharedPtr<FSocket>& InClientSocket, const TArray<uint8>& InBytes)
 {
 	UE_LOG(LogStormSyncServer, Display, TEXT("Handle received buffer via %s (Buffer Size: %d)"), *InEndpoint.ToString(), InBytes.Num());
-	ReceivedBufferEvent.Broadcast(InEndpoint, InClientSocket, MakeShared<TArray<uint8>>(InBytes));
+	ReceivedBufferEvent.Broadcast(InEndpoint, InClientSocket, MakeShared<FStormSyncBuffer>(InBytes));
 
 	// Fully received buffer, clean up connection and the state associated with it
 	PendingConnectionsToDisconnect.Enqueue(InEndpoint);
diff --git a/Engine/Plugins/Experimental/StormSync/Source/StormSyncTransportServer/Private/StormSyncTransportServerEndpoint.cpp b/Engine/Plugins/Experimental/StormSync/Source/StormSyncTransportServer/Private/StormSyncTransportServerEndpoint.cpp
--- a/Engine/Plugins/Experimental/StormSync/Source/StormSyncTransportServer/Private/StormSyncTransportServerEndpoint.cpp
+++ b/Engine/Plugins/Experimental/StormSync/Source/StormSyncTransportServer/Private/StormSyncTransportServerEndpoint.cpp
@@ -8,6 +8,7 @@
 #include "Socket/StormSyncTransportTcpServer.h"
 #include "StormSyncCoreDelegates.h"
 #include "StormSyncCoreUtils.h"
+#include "StormSyncMemoryReader.h"
 #include "StormSyncTransportNetworkUtils.h"
 #include "StormSyncTransportServerLog.h"
 #include "StormSyncTransportSettings.h"
@@ -200,7 +201,7 @@ void FStormSyncTransportServerEndpoint::HandleReceivedTcpBuffer(const FIPv4Endpo
 		return;
 	}
 
-	UE_LOG(LogStormSyncServer, Display, TEXT("Handle received tcp buffer via %s:%d (Buffer Size: %d)"), *InEndpoint.Address.ToString(), InEndpoint.Port, InBuffer->Num());
+	UE_LOG(LogStormSyncServer, Display, TEXT("Handle received tcp buffer via %s:%d (Buffer Size: %lld)"), *InEndpoint.Address.ToString(), InEndpoint.Port, InBuffer->Num());
 
 	// Local copy of the endpoint we're receiving a buffer from
 	FIPv4Endpoint Endpoint = InEndpoint;
@@ -214,7 +215,8 @@ void FStormSyncTransportServerEndpoint::HandleReceivedTcpBuffer(const FIPv4Endpo
 			FText::FromString(Endpoint.Address.ToString())
 		).ToString();
 
-		FStormSyncCoreDelegates::OnRequestImportBuffer.Broadcast(PackageDescriptor, InBuffer);
+		const FStormSyncArchivePtr MemoryReader = MakeShared<FStormSyncMemoryReader>(InBuffer);
+		FStormSyncCoreDelegates::OnRequestImportBuffer.Broadcast(PackageDescriptor, MemoryReader);
 	});
 }
 
