Fixed FText time zone double-conversions
`FText::AsDateTime` expects a UTC time parameter if `TimeZone` is not specified, and will convert to local time unless the `TimeZone` parameter is `FText::GetInvariantTimeZone()`.
This means any call to `AsDateTime` passing `FDateTime::Now()` instead of `FDateTime::UtcNow()` will get the offset from UTC re-applied (i.e., if in GMT-6, passing already-local time will result in GMT-12).

#12822

diff --git a/Engine/Plugins/Experimental/AnimNext/Source/AnimNextUncookedOnly/Private/AnimNextScopedCompilerResults.cpp b/Engine/Plugins/Experimental/AnimNext/Source/AnimNextUncookedOnly/Private/AnimNextScopedCompilerResults.cpp
--- a/Engine/Plugins/Experimental/AnimNext/Source/AnimNextUncookedOnly/Private/AnimNextScopedCompilerResults.cpp
+++ b/Engine/Plugins/Experimental/AnimNext/Source/AnimNextUncookedOnly/Private/AnimNextScopedCompilerResults.cpp
@@ -81,7 +81,7 @@ FScopedCompilerResults::~FScopedCompilerResults()
 		{
 			FMessageLog MessageLog("AnimNextCompilerResults");
 
-			MessageLog.NewPage(FText::Format(LOCTEXT("CompileFormat", "Compile {0}: {1}"), JobName, FText::AsDateTime(FDateTime::Now())));
+			MessageLog.NewPage(FText::Format(LOCTEXT("CompileFormat", "Compile {0}: {1}"), JobName, FText::AsDateTime(FDateTime::UtcNow())));
 			MessageLog.AddMessages(ThreadData->Messages);
 
 			// Print summary
diff --git a/Engine/Plugins/Experimental/EditorPerformance/Source/StallLogSubsystem/Private/StallLogSubsystem.cpp b/Engine/Plugins/Experimental/EditorPerformance/Source/StallLogSubsystem/Private/StallLogSubsystem.cpp
--- a/Engine/Plugins/Experimental/EditorPerformance/Source/StallLogSubsystem/Private/StallLogSubsystem.cpp
+++ b/Engine/Plugins/Experimental/EditorPerformance/Source/StallLogSubsystem/Private/StallLogSubsystem.cpp
@@ -236,7 +236,7 @@ namespace
 					[
 						SNew(STextBlock)
 							.ColorAndOpacity(FSlateColor::UseForeground())
-							.Text(FText::Format(LOCTEXT("TimeFmt", "{0}"), FText::AsDateTime(StallLogItem->Time)))
+							.Text(FText::Format(LOCTEXT("TimeFmt", "{0}"), FText::AsDateTime(StallLogItem->Time, EDateTimeStyle::Default, EDateTimeStyle::Default, FText::GetInvariantTimeZone())))
 					];
 			}
 			else if (ColumnName == ColumnName_Copy)
diff --git a/Engine/Plugins/Experimental/StormSync/Source/StormSyncEditor/Private/Subsystems/StormSyncNotificationSubsystem.cpp b/Engine/Plugins/Experimental/StormSync/Source/StormSyncEditor/Private/Subsystems/StormSyncNotificationSubsystem.cpp
--- a/Engine/Plugins/Experimental/StormSync/Source/StormSyncEditor/Private/Subsystems/StormSyncNotificationSubsystem.cpp
+++ b/Engine/Plugins/Experimental/StormSync/Source/StormSyncEditor/Private/Subsystems/StormSyncNotificationSubsystem.cpp
@@ -83,7 +83,7 @@ void UStormSyncNotificationSubsystem::NewPage(const FText& InLabel, const bool b
 	
 	if (bInAppendDateTime)
 	{
-		const FText PageLabel = FText::Format(LOCTEXT("PageLabel_DateTime_Format", "{0} - {1}"), InLabel, FText::AsDateTime(FDateTime::Now()));
+		const FText PageLabel = FText::Format(LOCTEXT("PageLabel_DateTime_Format", "{0} - {1}"), InLabel, FText::AsDateTime(FDateTime::UtcNow()));
 		MessageLog->NewPage(PageLabel);
 	}
 	else
diff --git a/Engine/Plugins/Experimental/StormSync/Source/StormSyncImport/Private/Subsystems/StormSyncImportSubsystem.cpp b/Engine/Plugins/Experimental/StormSync/Source/StormSyncImport/Private/Subsystems/StormSyncImportSubsystem.cpp
--- a/Engine/Plugins/Experimental/StormSync/Source/StormSyncImport/Private/Subsystems/StormSyncImportSubsystem.cpp
+++ b/Engine/Plugins/Experimental/StormSync/Source/StormSyncImport/Private/Subsystems/StormSyncImportSubsystem.cpp
@@ -466,7 +466,7 @@ void UStormSyncImportSubsystem::HandlePakPreExtract(const FStormSyncPackageDescr
 	// Init a new page for message log to categorize each pak we receive
 	FMessageLog MessageLog(LogName);
 
-	const FText PageMessage = FText::Format(LOCTEXT("PakPreExtract_Incoming_Pak", "Handle incoming pak \"{0}\" ({1})"), FText::FromString(InPackageDescriptor.Name), FText::AsDateTime(FDateTime::Now()));
+	const FText PageMessage = FText::Format(LOCTEXT("PakPreExtract_Incoming_Pak", "Handle incoming pak \"{0}\" ({1})"), FText::FromString(InPackageDescriptor.Name), FText::AsDateTime(FDateTime::UtcNow()));
 	MessageLog.NewPage(PageMessage);
 	MessageLog.Info(FText::FromString(InPackageDescriptor.ToString()));
 #endif
diff --git a/Engine/Source/Editor/UnrealEd/Private/EditorServer.cpp b/Engine/Source/Editor/UnrealEd/Private/EditorServer.cpp
--- a/Engine/Source/Editor/UnrealEd/Private/EditorServer.cpp
+++ b/Engine/Source/Editor/UnrealEd/Private/EditorServer.cpp
@@ -3867,7 +3867,7 @@ bool UEditorEngine::Map_Check( UWorld* InWorld, const TCHAR* Str, FOutputDevice&
 	{
 		FFormatNamedArguments Arguments;
 		Arguments.Add(TEXT("Name"), FText::FromString(FPackageName::GetShortName(InWorld->GetOutermost())));
-		Arguments.Add(TEXT("TimeStamp"), FText::AsDateTime(FDateTime::Now()));
+		Arguments.Add(TEXT("TimeStamp"), FText::AsDateTime(FDateTime::UtcNow()));
 		FText MapCheckPageName(FText::Format(LOCTEXT("MapCheckPageName", "{Name} - {TimeStamp}"), Arguments));
 		MapCheckLog.NewPage(MapCheckPageName);
 	}
diff --git a/Engine/Source/Editor/UnrealEd/Private/Kismet2/CompilerResultsLog.cpp b/Engine/Source/Editor/UnrealEd/Private/Kismet2/CompilerResultsLog.cpp
--- a/Engine/Source/Editor/UnrealEd/Private/Kismet2/CompilerResultsLog.cpp
+++ b/Engine/Source/Editor/UnrealEd/Private/Kismet2/CompilerResultsLog.cpp
@@ -610,7 +610,7 @@ void FCompilerResultsLog::GetGlobalModuleCompilerDump(const FString& LogDump, EC
 	FMessageLog MessageLog(Name);
 
 	FFormatNamedArguments Arguments;
-	Arguments.Add(TEXT("TimeStamp"), FText::AsDateTime(FDateTime::Now()));
+	Arguments.Add(TEXT("TimeStamp"), FText::AsDateTime(FDateTime::UtcNow()));
 	MessageLog.NewPage(FText::Format(LOCTEXT("CompilerLogPage", "Compilation - {TimeStamp}"), Arguments));
 
 	if ( bShowLog )
diff --git a/Engine/Source/Editor/UnrealEd/Private/PlayLevel.cpp b/Engine/Source/Editor/UnrealEd/Private/PlayLevel.cpp
--- a/Engine/Source/Editor/UnrealEd/Private/PlayLevel.cpp
+++ b/Engine/Source/Editor/UnrealEd/Private/PlayLevel.cpp
@@ -2699,7 +2699,7 @@ void UEditorEngine::StartPlayInEditorSession(FRequestPlaySessionParams& InReques
 
 		FFormatNamedArguments Arguments;
 		Arguments.Add(TEXT("Package"), FText::FromString(FPackageName::GetLongPackageAssetName(WorldPackageName)));
-		Arguments.Add(TEXT("TimeStamp"), FText::AsDateTime(FDateTime::Now()));
+		Arguments.Add(TEXT("TimeStamp"), FText::AsDateTime(FDateTime::UtcNow()));
 
 		FText PIESessionLabel = InRequestParams.WorldType == EPlaySessionWorldType::SimulateInEditor ?
 			FText::Format(LOCTEXT("SIESessionLabel", "SIE session: {Package} ({TimeStamp})"), Arguments) :
diff --git a/Engine/Source/Editor/UnrealEd/Private/StaticLightingSystem/StaticLightingSystem.cpp b/Engine/Source/Editor/UnrealEd/Private/StaticLightingSystem/StaticLightingSystem.cpp
--- a/Engine/Source/Editor/UnrealEd/Private/StaticLightingSystem/StaticLightingSystem.cpp
+++ b/Engine/Source/Editor/UnrealEd/Private/StaticLightingSystem/StaticLightingSystem.cpp
@@ -550,7 +550,7 @@ bool FStaticLightingSystem::BeginLightmassProcess()
 
 		// Flip the results page
 		FFormatNamedArguments Arguments;
-		Arguments.Add(TEXT("TimeStamp"), FText::AsDateTime(FDateTime::Now()));
+		Arguments.Add(TEXT("TimeStamp"), FText::AsDateTime(FDateTime::UtcNow()));
 		FText LightingResultsPageName(FText::Format(LOCTEXT("LightingResultsPageName", "Lighting Build - {TimeStamp}"), Arguments));
 		FMessageLog("LightingResults").NewPage(LightingResultsPageName);
 
