[StormSync] Follow up to CL 41376663. Addressing minor comments from code review.

diff --git a/Engine/Plugins/Experimental/StormSync/Source/StormSyncCore/Internal/StormSyncMemoryReader.h b/Engine/Plugins/Experimental/StormSync/Source/StormSyncCore/Internal/StormSyncMemoryReader.h
--- a/Engine/Plugins/Experimental/StormSync/Source/StormSyncCore/Internal/StormSyncMemoryReader.h
+++ b/Engine/Plugins/Experimental/StormSync/Source/StormSyncCore/Internal/StormSyncMemoryReader.h
@@ -44,8 +44,8 @@ public:
 	explicit FStormSyncMemoryReader(const FStormSyncBufferPtr& InBuffer)
 		: Buffer(InBuffer)
 	{
-		this->SetIsLoading(true);
-		this->SetIsPersistent(false);
+		SetIsLoading(true);
+		SetIsPersistent(false);
 	}
 
 private:
diff --git a/Engine/Plugins/Experimental/StormSync/Source/StormSyncCore/Private/StormSyncCoreUtils.cpp b/Engine/Plugins/Experimental/StormSync/Source/StormSyncCore/Private/StormSyncCoreUtils.cpp
--- a/Engine/Plugins/Experimental/StormSync/Source/StormSyncCore/Private/StormSyncCoreUtils.cpp
+++ b/Engine/Plugins/Experimental/StormSync/Source/StormSyncCore/Private/StormSyncCoreUtils.cpp
@@ -26,7 +26,7 @@ namespace UE::StormSyncCore::Private
 	// physical device's page size.
 	constexpr uint64 ChunkBufferSize = PLATFORM_FILE_WRITER_BUFFER_SIZE;
 
-	uint64 GetChunkBufferSize()
+	constexpr uint64 GetChunkBufferSize()
 	{
 		return ChunkBufferSize;
 	}
@@ -35,7 +35,7 @@ namespace UE::StormSyncCore::Private
 	void ArchiveCopy(FArchive& InReadArchive, FArchive& InWriteArchive, int64 InSizeToCopy, TArray64<uint8>& InChunkBuffer)
 	{
 		// Note: This should already be validated prior, but we ensure here as a precaution.
-		if(!ensure(InReadArchive.IsLoading()) || !ensure(InWriteArchive.IsSaving()))
+		if (!ensure(InReadArchive.IsLoading()) || !ensure(InWriteArchive.IsSaving()))
 		{
 			return;
 		}
@@ -236,7 +236,9 @@ bool FStormSyncCoreUtils::CreatePakBuffer(const TArray<FName>& InPackageNames, F
 
 	if (!OutPakArchive.IsSaving())
 	{
-		OutErrorText = FText::FromString(FString::Printf(TEXT("CreatePakBuffer: Archive \"%s\" is not writable."), *OutPakArchive.GetArchiveName()));
+		OutErrorText = FText::Format(
+			LOCTEXT("CreatePakBuffer_ArchiveNotWritable", "CreatePakBuffer: Archive \"{0}\" is not writable."),
+			FText::FromString(OutPakArchive.GetArchiveName()));
 		return false;
 	}
 
diff --git a/Engine/Plugins/Experimental/StormSync/Source/StormSyncImport/Private/Subsystems/StormSyncImportSubsystem.h b/Engine/Plugins/Experimental/StormSync/Source/StormSyncImport/Private/Subsystems/StormSyncImportSubsystem.h
--- a/Engine/Plugins/Experimental/StormSync/Source/StormSyncImport/Private/Subsystems/StormSyncImportSubsystem.h
+++ b/Engine/Plugins/Experimental/StormSync/Source/StormSyncImport/Private/Subsystems/StormSyncImportSubsystem.h
@@ -90,7 +90,7 @@ public:
 	 * the list of files that needs to be extracted and updated in this editor instance.
 	 *
 	 * @param InPackageDescriptor The package descriptor for this buffer extraction. Only used for UI message logs.
-	 * @param InArchive Source archive to be parse to extract content
+	 * @param InArchive Source archive to be parsed to extract content
 	 * @param bShowWizard Whether or not to show the import wizard
 	 * @param bDryRun Whether we should proceed to extraction or simply log to console buffer content
 	 * 
