Motion Design: moved the following plugins from /Plugins/Experimental to /Plugins/VirtualProduction:

diff --git a/Engine/Plugins/VirtualProduction/Avalanche/Source/AvalancheTransitionEditor/Private/AvaTransitionEditorModule.cpp b/Engine/Plugins/VirtualProduction/Avalanche/Source/AvalancheTransitionEditor/Private/AvaTransitionEditorModule.cpp
new file mode 100644
--- /dev/null
+++ b/Engine/Plugins/VirtualProduction/Avalanche/Source/AvalancheTransitionEditor/Private/AvaTransitionEditorModule.cpp
@@ -0,0 +1,179 @@
+// Copyright Epic Games, Inc. All Rights Reserved.
+
+#include "AvaTransitionEditorModule.h"
+#include "AvaTransitionCommands.h"
+#include "AvaTransitionEditorEnums.h"
+#include "AvaTransitionEditorLog.h"
+#include "AvaTransitionEditorUtils.h"
+#include "AvaTransitionTree.h"
+#include "AvaTransitionTreeEditorData.h"
+#include "AvaTransitionTreeSchema.h"
+#include "Behavior/IAvaTransitionBehavior.h"
+#include "IAvaTransitionModule.h"
+#include "StateTreeDelegates.h"
+#include "Styling/AvaTransitionEditorStyle.h"
+#include "ToolMenu.h"
+#include "ToolMenuSection.h"
+#include "Settings/AvaTransitionEditorSettings.h"
+#include "ViewModels/AvaTransitionEditorViewModel.h"
+#include "Widgets/Layout/SBox.h"
+
+#if WITH_STATETREE_DEBUGGER
+#include "Debugger/AvaTransitionTraceModule.h"
+#endif
+
+#define LOCTEXT_NAMESPACE "AvaTransitionEditorModule"
+
+DEFINE_LOG_CATEGORY(LogAvaEditorTransition)
+
+IMPLEMENT_MODULE(FAvaTransitionEditorModule, AvalancheTransitionEditor)
+
+void FAvaTransitionEditorModule::StartupModule()
+{
+#if WITH_STATETREE_DEBUGGER
+	FAvaTransitionTraceModule::Startup();
+#endif
+
+	FAvaTransitionEditorCommands::Register();
+
+	IAvaTransitionModule::Get().GetOnValidateTransitionTree().BindRaw(this, &FAvaTransitionEditorModule::ValidateStateTree);
+
+	OnPostCompileHandle = UE::StateTree::Delegates::OnPostCompile.AddRaw(this, &FAvaTransitionEditorModule::OnPostCompile);
+}
+
+void FAvaTransitionEditorModule::ShutdownModule()
+{
+	FAvaTransitionEditorCommands::Unregister();
+
+#if WITH_STATETREE_DEBUGGER
+	FAvaTransitionTraceModule::Shutdown();
+#endif
+
+	IAvaTransitionModule::Get().GetOnValidateTransitionTree().Unbind();
+
+	UE::StateTree::Delegates::OnPostCompile.Remove(OnPostCompileHandle);
+	OnPostCompileHandle.Reset();
+}
+
+IAvaTransitionEditorModule::FOnBuildDefaultTransitionTree& FAvaTransitionEditorModule::GetOnBuildDefaultTransitionTree()
+{
+	return OnBuildDefaultTransitionTree;
+}
+
+void FAvaTransitionEditorModule::GenerateTransitionTreeOptionsMenu(UToolMenu* InMenu, IAvaTransitionBehavior* InTransitionBehavior)
+{
+	if (!InMenu || !InTransitionBehavior)
+	{
+		return;
+	}
+
+	UAvaTransitionTree* TransitionTree = InTransitionBehavior->GetTransitionTree();
+	if (!TransitionTree)
+	{
+		return;
+	}
+
+	TWeakObjectPtr<UAvaTransitionTree> TransitionTreeWeak = TransitionTree;
+
+	FToolMenuSection& GeneralSection = InMenu->FindOrAddSection(TEXT("GeneralSection"), LOCTEXT("GeneralSectionLabel", "General"));
+
+	GeneralSection.AddMenuEntry(TEXT("TransitionTreeEnabled")
+		, LOCTEXT("TransitionTreeEnabledLabel", "Enabled")
+		, LOCTEXT("TransitionTreeEnabledTooltip", "Determines whether Transition Tree is Enabled")
+		, TAttribute<FSlateIcon>()
+		, FUIAction(FExecuteAction::CreateStatic(&UE::AvaTransitionEditor::ToggleTransitionTreeEnabled, TransitionTreeWeak)
+			, FCanExecuteAction()
+			, FIsActionChecked::CreateStatic(&UE::AvaTransitionEditor::IsTransitionTreeEnabled, TransitionTreeWeak))
+		, EUserInterfaceActionType::ToggleButton);
+
+	if (const TSharedPtr<SWidget> ModeSelector = UE::AvaTransitionEditor::CreateTransitionInstancingModeSelector(TransitionTree))
+	{
+		GeneralSection.AddEntry(FToolMenuEntry::InitWidget(TEXT("TransitionModeSelector")
+			, SNew(SBox)
+				[
+					ModeSelector.ToSharedRef()
+				]
+			, FText::GetEmpty()
+			, /*bNoIndent*/true
+			, /*bSearchable*/false
+			, /*bNoPadding*/true));
+	}
+
+	if (TSharedPtr<SWidget> LayerPicker = UE::AvaTransitionEditor::CreateTransitionLayerPicker(Cast<UAvaTransitionTreeEditorData>(TransitionTree->EditorData), /*bInCompileOnLayerPicked*/true))
+	{
+		FToolMenuSection& LayerSection = InMenu->FindOrAddSection(TEXT("LayerSection"), LOCTEXT("LayerSectionLabel", "Layer"));
+
+		LayerSection.AddEntry(FToolMenuEntry::InitWidget(TEXT("TransitionLayerPicker")
+			, SNew(SBox)
+				.HeightOverride(40.f)
+				[
+					LayerPicker.ToSharedRef()
+				]
+			, FText::GetEmpty()
+			, /*bNoIndent*/true
+			, /*bSearchable*/false
+			, /*bNoPadding*/true));
+	}
+}
+
+void FAvaTransitionEditorModule::ValidateStateTree(UAvaTransitionTree* InTransitionTree)
+{
+	// Return early if Editor Data is already valid
+	if (!InTransitionTree || InTransitionTree->EditorData)
+	{
+		return;
+	}
+
+	// Disable Tree by default if being set up for the first time
+	InTransitionTree->SetEnabled(false);
+
+	UAvaTransitionTreeEditorData* EditorData;
+
+	const UAvaTransitionEditorSettings* TransitionEditorSettings = GetDefault<UAvaTransitionEditorSettings>();
+	check(TransitionEditorSettings);
+
+	if (UAvaTransitionTreeEditorData* TemplateEditorData = TransitionEditorSettings->LoadDefaultTemplateEditorData())
+	{
+		EditorData = DuplicateObject<UAvaTransitionTreeEditorData>(TemplateEditorData, InTransitionTree);
+		check(EditorData);
+	}
+	else
+	{
+		EditorData = NewObject<UAvaTransitionTreeEditorData>(InTransitionTree, NAME_None, RF_Transactional);
+		check(EditorData);
+
+		EditorData->Schema = NewObject<UAvaTransitionTreeSchema>(EditorData);
+
+		if (OnBuildDefaultTransitionTree.IsBound())
+		{
+			OnBuildDefaultTransitionTree.Execute(*EditorData);
+		}
+		else
+		{
+			EditorData->AddRootState();	
+		}
+	}
+
+	InTransitionTree->EditorData = EditorData;
+
+	// Compile in Advanced Mode here so that no new nodes are generated from outside
+	FAvaTransitionCompiler Compiler;
+	Compiler.SetTransitionTree(InTransitionTree);
+	Compiler.Compile(EAvaTransitionEditorMode::Advanced);
+}
+
+void FAvaTransitionEditorModule::OnPostCompile(const UStateTree& InStateTree)
+{
+	// TODO: Currently State Tree Compiler does not have any hooks to compile additional data,
+	// so this is where we inject the additional properties
+	if (const UAvaTransitionTree* TransitionTree = Cast<UAvaTransitionTree>(&InStateTree))
+	{
+		UAvaTransitionTree* TransitionTreeMutable = const_cast<UAvaTransitionTree*>(TransitionTree);
+
+		const UAvaTransitionTreeEditorData* EditorData = CastChecked<UAvaTransitionTreeEditorData>(InStateTree.EditorData);
+
+		TransitionTreeMutable->SetTransitionLayer(EditorData->GetTransitionLayer());
+	}
+}
+
+#undef LOCTEXT_NAMESPACE
diff --git a/Engine/Plugins/VirtualProduction/Avalanche/Source/AvalancheTransitionEditor/Private/AvaTransitionEditorUtils.h b/Engine/Plugins/VirtualProduction/Avalanche/Source/AvalancheTransitionEditor/Private/AvaTransitionEditorUtils.h
new file mode 100644
--- /dev/null
+++ b/Engine/Plugins/VirtualProduction/Avalanche/Source/AvalancheTransitionEditor/Private/AvaTransitionEditorUtils.h
@@ -0,0 +1,28 @@
+// Copyright Epic Games, Inc. All Rights Reserved.
+
+#pragma once
+
+#include "Behavior/IAvaTransitionBehavior.h"
+#include "Templates/SharedPointerFwd.h"
+#include "UObject/WeakObjectPtr.h"
+
+class SWidget;
+class UAvaTransitionTree;
+class UAvaTransitionTreeEditorData;
+
+namespace UE::AvaTransitionEditor
+{
+	TSharedPtr<SWidget> CreateTransitionLayerPicker(UAvaTransitionTreeEditorData* InEditorData, bool bInCompileOnLayerPicked);
+
+	TSharedPtr<SWidget> CreateTransitionInstancingModeSelector(UAvaTransitionTree* InTransitionTree);
+	
+	void ValidateTree(UAvaTransitionTree& InTransitionTree);
+
+	uint32 CalculateTreeHash(UAvaTransitionTree& InTransitionTree);
+
+	bool PickTransitionTreeAsset(const FText& InDialogTitle, UAvaTransitionTree*& OutTransitionTree);
+
+	void ToggleTransitionTreeEnabled(TWeakObjectPtr<UAvaTransitionTree> InTransitionTreeWeak);
+
+	bool IsTransitionTreeEnabled(TWeakObjectPtr<UAvaTransitionTree> InTransitionTreeWeak);
+}
diff --git a/Engine/Plugins/VirtualProduction/Avalanche/Source/AvalancheTransitionEditor/Private/Debugger/AvaTransitionTraceAnalyzer.cpp b/Engine/Plugins/VirtualProduction/Avalanche/Source/AvalancheTransitionEditor/Private/Debugger/AvaTransitionTraceAnalyzer.cpp
new file mode 100644
--- /dev/null
+++ b/Engine/Plugins/VirtualProduction/Avalanche/Source/AvalancheTransitionEditor/Private/Debugger/AvaTransitionTraceAnalyzer.cpp
@@ -0,0 +1,237 @@
+// Copyright Epic Games, Inc. All Rights Reserved.
+
+#if WITH_STATETREE_DEBUGGER
+
+#include "AvaTransitionTraceAnalyzer.h"
+#include "AvaTransitionDebugger.h"
+#include "AvaTransitionEditorLog.h"
+#include "Containers/Ticker.h"
+#include "Debugger/StateTreeDebuggerTypes.h"
+#include "StateTree.h"
+#include "StateTreeExecutionTypes.h"
+#include "StateTreeModule.h"
+#include "TraceServices/Model/AnalysisSession.h"
+
+namespace UE::AvaTransitionEditor::Private
+{
+	EStateTreeTraceEventType GetEventType(const Trace::IAnalyzer::FOnEventContext& InContext)
+	{
+		return InContext.EventData.GetValue<EStateTreeTraceEventType>("EventType");
+	}
+
+	FStateTreeIndex16 GetAssetDebugId(const Trace::IAnalyzer::FOnEventContext& InContext)
+	{
+		return FStateTreeIndex16(InContext.EventData.GetValue<uint16>("AssetDebugId"));
+	}
+
+	FStateTreeInstanceDebugId GetInstanceDebugId(const Trace::IAnalyzer::FOnEventContext& InContext)
+	{
+		return FStateTreeInstanceDebugId(InContext.EventData.GetValue<uint32>("InstanceId"), InContext.EventData.GetValue<uint32>("InstanceSerial"));
+	}
+}
+
+TArray<FAvaTransitionDebugger*> FAvaTransitionTraceAnalyzer::ActiveDebuggers;
+
+const TArray<FAvaTransitionTraceAnalyzer::FEvent> FAvaTransitionTraceAnalyzer::Events =
+{
+	{ "AssetDebugIdEvent", &FAvaTransitionTraceAnalyzer::OnAssetDebugIdEvent },
+	{ "InstanceEvent"    , &FAvaTransitionTraceAnalyzer::OnInstanceEvent },
+	{ "StateEvent"       , &FAvaTransitionTraceAnalyzer::OnStateEvent },
+};
+
+void FAvaTransitionTraceAnalyzer::RegisterDebugger(FAvaTransitionDebugger* InDebugger)
+{
+	ActiveDebuggers.Add(InDebugger);
+
+	if (IStateTreeModule::IsAvailable())
+	{
+		int32 TraceId = 0;
+        IStateTreeModule::Get().StartTraces(TraceId);
+	}
+}
+
+void FAvaTransitionTraceAnalyzer::UnregisterDebugger(FAvaTransitionDebugger* InDebugger)
+{
+	bool bRemovedDebugger = ActiveDebuggers.Remove(InDebugger) > 0;
+
+	// Stop Traces only if the passed in Debugger was the last of the Active Debuggers
+	if (bRemovedDebugger && ActiveDebuggers.IsEmpty() && IStateTreeModule::IsAvailable())
+	{
+		IStateTreeModule::Get().StopTraces();
+	}
+}
+
+void FAvaTransitionTraceAnalyzer::OnAnalysisBegin(const FOnAnalysisContext& InContext)
+{
+	for (int32 RouteId = 0; RouteId < Events.Num(); ++RouteId)
+	{
+		InContext.InterfaceBuilder.RouteEvent(RouteId, "StateTreeDebugger", Events[RouteId].Name);
+	}
+}
+
+bool FAvaTransitionTraceAnalyzer::OnEvent(uint16 InRouteId, EStyle InStyle, const FOnEventContext& InContext)
+{
+	if (Events.IsValidIndex(InRouteId))
+	{
+		::Invoke(Events[InRouteId].Handler, this, InContext);
+	}
+	return true;
+}
+
+void FAvaTransitionTraceAnalyzer::ForEachDebugger(const TCHAR* InDebugName, TFunction<void(FAvaTransitionDebugger&)>&& InCallable)
+{
+	ExecuteOnGameThread(InDebugName,
+		[Callable = MoveTemp(InCallable)]
+		{
+			for (FAvaTransitionDebugger* Debugger : ActiveDebuggers)
+			{
+				check(Debugger);
+				Callable(*Debugger);
+			}
+		});
+}
+
+UStateTree* FAvaTransitionTraceAnalyzer::FindStateTree(const FStateTreeInstanceDebugId& InDebugInstanceId) const
+{
+	const FStateTreeIndex16* AssetDebugId = DebugInstances.Find(InDebugInstanceId);
+	if (!AssetDebugId)
+	{
+		return nullptr;
+	}
+
+	const TWeakObjectPtr<UStateTree>* DebugAsset = DebugAssets.Find(*AssetDebugId);
+	if (!DebugAsset)
+	{
+		return nullptr;
+	}
+
+	return DebugAsset->Get();
+}
+
+void FAvaTransitionTraceAnalyzer::OnAssetDebugIdEvent(const FOnEventContext& InContext)
+{
+	FString TreeName, TreePathName;
+	InContext.EventData.GetString("TreeName", TreeName);
+	InContext.EventData.GetString("TreePath", TreePathName);
+
+	UE_LOG(LogAvaEditorTransition, Verbose, TEXT("OnAssetDebugIdEvent --- Attempting to find tree %s in path %s"), *TreeName, *TreePathName);
+
+	UStateTree* StateTree;
+	{
+		FGCScopeGuard Guard;
+		StateTree = FindObject<UStateTree>(nullptr, *TreePathName);
+	}
+
+	const uint32 CompiledDataHash = InContext.EventData.GetValue<uint32>("CompiledDataHash");
+	if (StateTree && StateTree->LastCompiledEditorDataHash == CompiledDataHash)
+	{
+		FStateTreeIndex16 AssetDebugId = UE::AvaTransitionEditor::Private::GetAssetDebugId(InContext);
+		DebugAssets.Add(AssetDebugId, StateTree);
+	}
+	else
+	{
+		UE_LOG(LogAvaEditorTransition, Warning, TEXT("Unable to find StateTree asset: %s : %s"), *TreePathName, *TreeName);
+	}
+}
+
+void FAvaTransitionTraceAnalyzer::OnInstanceEvent(const FOnEventContext& InContext)
+{
+	EStateTreeTraceEventType EventType = UE::AvaTransitionEditor::Private::GetEventType(InContext);
+
+	switch (EventType)
+	{
+	case EStateTreeTraceEventType::Push:
+		OnTreeInstanceStarted(InContext);
+		break;
+	case EStateTreeTraceEventType::Pop:
+		OnTreeInstanceStopped(InContext);
+		break;
+	}
+}
+
+void FAvaTransitionTraceAnalyzer::OnTreeInstanceStarted(const FOnEventContext& InContext)
+{
+	using namespace UE::AvaTransitionEditor;
+
+	FStateTreeInstanceDebugId InstanceId = UE::AvaTransitionEditor::Private::GetInstanceDebugId(InContext);
+	FStateTreeIndex16 AssetId = UE::AvaTransitionEditor::Private::GetAssetDebugId(InContext);
+
+	DebugInstances.Add(InstanceId, AssetId);
+
+	FString InstanceName;
+	InContext.EventData.GetString("InstanceName", InstanceName);
+
+	if (TWeakObjectPtr<UStateTree>* DebugAsset = DebugAssets.Find(AssetId))
+	{
+		ForEachDebugger(UE_SOURCE_LOCATION,
+			[StateTreeWeak = *DebugAsset, InstanceId, InstanceName](FAvaTransitionDebugger& InDebugger)
+			{
+				if (UStateTree* StateTree = StateTreeWeak.Get())
+				{
+					InDebugger.OnTreeInstanceStarted(*StateTree, InstanceId, InstanceName);
+				}
+			});
+	}
+}
+
+void FAvaTransitionTraceAnalyzer::OnTreeInstanceStopped(const FOnEventContext& InContext)
+{
+	FStateTreeInstanceDebugId InstanceDebugId = UE::AvaTransitionEditor::Private::GetInstanceDebugId(InContext);
+
+	DebugInstances.Remove(InstanceDebugId);
+
+	ForEachDebugger(UE_SOURCE_LOCATION,
+		[InstanceDebugId](FAvaTransitionDebugger& InDebugger)
+		{
+			InDebugger.OnTreeInstanceStopped(InstanceDebugId);
+		});
+}
+
+void FAvaTransitionTraceAnalyzer::OnStateEvent(const FOnEventContext& InContext)
+{
+	EStateTreeTraceEventType EventType = UE::AvaTransitionEditor::Private::GetEventType(InContext);
+
+	// Ignore other events for now. Only concerned with enter/exit
+	if (EventType != EStateTreeTraceEventType::OnEntered && EventType != EStateTreeTraceEventType::OnExited)
+	{
+		return;
+	}
+
+	FStateTreeInstanceDebugId InstanceDebugId = UE::AvaTransitionEditor::Private::GetInstanceDebugId(InContext);
+	if (!InstanceDebugId.IsValid())
+	{
+		return;
+	}
+
+	UStateTree* StateTree = FindStateTree(InstanceDebugId);
+	if (!StateTree)
+	{
+		return;
+	}
+
+	FStateTreeStateHandle StateHandle(InContext.EventData.GetValue<uint16>("StateIndex"));
+
+	FGuid StateId = StateTree->GetStateIdFromHandle(StateHandle);
+	if (!StateId.IsValid())
+	{
+		UE_LOG(LogAvaEditorTransition, Warning, TEXT("Could not find id of state with index %d"), StateHandle.Index);
+		return;
+	}
+
+	ForEachDebugger(UE_SOURCE_LOCATION,
+		[EventType, StateId, InstanceDebugId](FAvaTransitionDebugger& InDebugger)
+		{
+			switch (EventType)
+			{
+			case EStateTreeTraceEventType::OnEntered:
+				InDebugger.OnNodeEntered(StateId, InstanceDebugId);
+				break;
+
+			case EStateTreeTraceEventType::OnExited:
+				InDebugger.OnNodeExited(StateId, InstanceDebugId);
+				break;
+			}
+		});
+}
+
+#endif // WITH_STATETREE_DEBUGGER
diff --git a/Engine/Plugins/VirtualProduction/Avalanche/Source/AvalancheTransitionEditor/Private/Debugger/AvaTransitionTraceAnalyzer.h b/Engine/Plugins/VirtualProduction/Avalanche/Source/AvalancheTransitionEditor/Private/Debugger/AvaTransitionTraceAnalyzer.h
new file mode 100644
--- /dev/null
+++ b/Engine/Plugins/VirtualProduction/Avalanche/Source/AvalancheTransitionEditor/Private/Debugger/AvaTransitionTraceAnalyzer.h
@@ -0,0 +1,58 @@
+// Copyright Epic Games, Inc. All Rights Reserved.
+
+#pragma once
+
+#if WITH_STATETREE_DEBUGGER
+
+#include "StateTreeExecutionTypes.h"
+#include "StateTreeIndexTypes.h"
+#include "Trace/Analyzer.h"
+#include "UObject/WeakObjectPtr.h"
+
+class FAvaTransitionDebugger;
+class UStateTree;
+
+class FAvaTransitionTraceAnalyzer : public UE::Trace::IAnalyzer
+{
+	struct FEvent
+	{
+		const ANSICHAR* Name;
+		void(FAvaTransitionTraceAnalyzer::*Handler)(const FOnEventContext&);
+	};
+
+public:
+	static void RegisterDebugger(FAvaTransitionDebugger* InDebugger);
+	static void UnregisterDebugger(FAvaTransitionDebugger* InDebugger);
+
+	//~ Begin IAnalyzer
+	virtual void OnAnalysisBegin(const FOnAnalysisContext& InContext) override;
+	virtual bool OnEvent(uint16 InRouteId, EStyle InStyle, const FOnEventContext& InContext) override;
+	//~ End IAnalyzer
+
+private:
+	static TArray<FAvaTransitionDebugger*> ActiveDebuggers;
+
+	static const TArray<FEvent> Events;
+
+	void ForEachDebugger(const TCHAR* InDebugName, TFunction<void(FAvaTransitionDebugger&)>&& InCallable);
+
+	UStateTree* FindStateTree(const FStateTreeInstanceDebugId& InDebugInstanceId) const;
+
+	void OnAssetDebugIdEvent(const FOnEventContext& InContext);
+
+	void OnInstanceEvent(const FOnEventContext& InContext);
+
+	void OnTreeInstanceStarted(const FOnEventContext& InContext);
+
+	void OnTreeInstanceStopped(const FOnEventContext& InContext);
+
+	void OnStateEvent(const FOnEventContext& InContext);
+
+	/** Mapping of the Asset Debug Id its State Tree */
+	TMap<FStateTreeIndex16, TWeakObjectPtr<UStateTree>> DebugAssets;
+
+	/** Mapping of the Instance Debug Id its Asset Debug Id */
+	TMap<FStateTreeInstanceDebugId, FStateTreeIndex16> DebugInstances;
+};
+
+#endif //WITH_STATETREE_DEBUGGER
diff --git a/Engine/Plugins/VirtualProduction/Avalanche/Source/AvalancheTransitionEditor/Private/Debugger/AvaTransitionTraceModule.cpp b/Engine/Plugins/VirtualProduction/Avalanche/Source/AvalancheTransitionEditor/Private/Debugger/AvaTransitionTraceModule.cpp
new file mode 100644
--- /dev/null
+++ b/Engine/Plugins/VirtualProduction/Avalanche/Source/AvalancheTransitionEditor/Private/Debugger/AvaTransitionTraceModule.cpp
@@ -0,0 +1,37 @@
+// Copyright Epic Games, Inc. All Rights Reserved.
+
+#if WITH_STATETREE_DEBUGGER
+
+#include "AvaTransitionTraceModule.h"
+#include "AvaTransitionTraceAnalyzer.h"
+#include "Features/IModularFeatures.h"
+#include "TraceServices/Model/AnalysisSession.h"
+
+FAvaTransitionTraceModule* FAvaTransitionTraceModule::Get()
+{
+	static FAvaTransitionTraceModule TraceModule;
+	return &TraceModule;
+}
+
+void FAvaTransitionTraceModule::Startup()
+{
+	IModularFeatures::Get().RegisterModularFeature(TraceServices::ModuleFeatureName, Get());
+}
+
+void FAvaTransitionTraceModule::Shutdown()
+{
+	IModularFeatures::Get().UnregisterModularFeature(TraceServices::ModuleFeatureName, Get());
+}
+
+void FAvaTransitionTraceModule::GetModuleInfo(TraceServices::FModuleInfo& OutModuleInfo)
+{
+	OutModuleInfo.Name = TEXT("TraceModule_AvaTransition");
+	OutModuleInfo.DisplayName = TEXT("Motion Design Transition Trace");
+}
+
+void FAvaTransitionTraceModule::OnAnalysisBegin(TraceServices::IAnalysisSession& InSession)
+{
+	InSession.AddAnalyzer(new FAvaTransitionTraceAnalyzer);
+}
+
+#endif // WITH_STATETREE_DEBUGGER
diff --git a/Engine/Plugins/VirtualProduction/Avalanche/Source/AvalancheTransitionEditor/Private/Debugger/AvaTransitionTraceModule.h b/Engine/Plugins/VirtualProduction/Avalanche/Source/AvalancheTransitionEditor/Private/Debugger/AvaTransitionTraceModule.h
new file mode 100644
--- /dev/null
+++ b/Engine/Plugins/VirtualProduction/Avalanche/Source/AvalancheTransitionEditor/Private/Debugger/AvaTransitionTraceModule.h
@@ -0,0 +1,25 @@
+// Copyright Epic Games, Inc. All Rights Reserved.
+
+#pragma once
+
+#if WITH_STATETREE_DEBUGGER
+
+#include "Trace/StoreClient.h"
+#include "TraceServices/ModuleService.h"
+
+class FAvaTransitionTraceModule : public TraceServices::IModule
+{
+	static FAvaTransitionTraceModule* Get();
+
+public:
+	static void Startup();
+	static void Shutdown();
+
+private:
+	//~ Begin IModule
+	virtual void GetModuleInfo(TraceServices::FModuleInfo& OutModuleInfo) override;
+	virtual void OnAnalysisBegin(TraceServices::IAnalysisSession& InSession) override;
+	//~ End IModule
+};
+
+#endif // WITH_STATETREE_DEBUGGER
diff --git a/Engine/Plugins/VirtualProduction/Avalanche/Source/AvalancheTransitionEditor/Public/AvaTransitionTreeEditorData.h b/Engine/Plugins/VirtualProduction/Avalanche/Source/AvalancheTransitionEditor/Public/AvaTransitionTreeEditorData.h
new file mode 100644
--- /dev/null
+++ b/Engine/Plugins/VirtualProduction/Avalanche/Source/AvalancheTransitionEditor/Public/AvaTransitionTreeEditorData.h
@@ -0,0 +1,73 @@
+// Copyright Epic Games, Inc. All Rights Reserved.
+
+#pragma once
+
+#include "AvaTagHandle.h"
+#include "Containers/Map.h"
+#include "Delegates/Delegate.h"
+#include "Misc/Guid.h"
+#include "State/AvaTransitionStateMetadata.h"
+#include "StateTreeEditorData.h"
+#include "AvaTransitionTreeEditorData.generated.h"
+
+class UAvaTransitionTreeEditorData;
+class UStateTreeState;
+
+namespace UE::AvaTransitionEditor
+{
+	constexpr FGuid ColorId_Default;
+	constexpr FGuid ColorId_In(0x1DDBC788, 0xD5EB400E, 0xBB71E5DA, 0xB27A784D);
+	constexpr FGuid ColorId_Out(0xE549EFA0, 0xDEFF45A7, 0xA8D907AF, 0xDB8F7643);
+}
+
+UCLASS(MinimalAPI, HideCategories=(Common))
+class UAvaTransitionTreeEditorData : public UStateTreeEditorData
+{
+	GENERATED_BODY()
+
+public:
+	UAvaTransitionTreeEditorData();
+
+	AVALANCHETRANSITIONEDITOR_API UStateTreeState& CreateState(const UStateTreeState& InSiblingState, bool bInAfter);
+
+	FAvaTagHandle GetTransitionLayer() const
+	{
+		return TransitionLayer;
+	}
+
+	void SetTransitionLayer(const FAvaTagHandle& InTransitionLayer)
+	{
+		TransitionLayer = InTransitionLayer;
+	}
+
+	static FName GetTransitionLayerPropertyName()
+	{
+		return GET_MEMBER_NAME_CHECKED(UAvaTransitionTreeEditorData, TransitionLayer);
+	}
+
+	const FAvaTransitionStateMetadata* FindStateMetadata(const FGuid& InStateId) const
+	{
+		return StateMetadata.Find(InStateId);
+	}
+
+	FAvaTransitionStateMetadata& FindOrAddStateMetadata(const FGuid& InStateId)
+	{
+		return StateMetadata.FindOrAdd(InStateId);
+	}
+
+	FSimpleMulticastDelegate& GetOnTreeRequestRefresh()
+	{
+		return OnTreeRequestRefresh;
+	}
+
+private:
+	/** The Layer this Transition Logic Tree deals with */
+	UPROPERTY(EditAnywhere, Category="Transition Logic")
+	FAvaTagHandle TransitionLayer;
+
+	/** Map of a state's id to its metadata */
+	UPROPERTY()
+	TMap<FGuid, FAvaTransitionStateMetadata> StateMetadata;
+
+	FSimpleMulticastDelegate OnTreeRequestRefresh;
+};
diff --git a/Engine/Plugins/VirtualProduction/Avalanche/Templates/TP_ME_BroadcastBP/TP_ME_BroadcastBP.uproject b/Engine/Plugins/VirtualProduction/Avalanche/Templates/TP_ME_BroadcastBP/TP_ME_BroadcastBP.uproject
new file mode 100644
--- /dev/null
+++ b/Engine/Plugins/VirtualProduction/Avalanche/Templates/TP_ME_BroadcastBP/TP_ME_BroadcastBP.uproject
@@ -0,0 +1,78 @@
+{
+	"FileVersion": 3,
+	"EngineAssociation": "",
+	"Category": "",
+	"Description": "",
+	"Plugins": [
+		{
+			"Name": "RemoteControlProtocolDMX",
+			"Enabled": true
+		},
+		{
+			"Name": "RemoteControlProtocolMIDI",
+			"Enabled": true
+		},
+		{
+			"Name": "RemoteControlProtocolOSC",
+			"Enabled": true
+		},
+		{
+			"Name": "DatasmithImporter",
+			"Enabled": true
+		},
+		{
+			"Name": "ColorCorrectRegions",
+			"Enabled": true
+		},
+		{
+			"Name": "RemoteControlWebInterface",
+			"Enabled": true,
+			"MarketplaceURL": "com.epicgames.launcher://ue/marketplace/content/beb39fbe47e74bb1a8098e2b2c9c40b2"
+		},
+		{
+			"Name": "VirtualProductionUtilities",
+			"Enabled": true
+		},
+		{
+			"Name": "RenderGrid",
+			"Enabled": true
+		},
+		{
+			"Name": "ElectraPlayer",
+			"Enabled": true
+		},
+		{
+			"Name": "ElectraUtil",
+			"Enabled": true
+		},
+		{
+			"Name": "BlackmagicMedia",
+			"Enabled": true,
+			"MarketplaceURL": "com.epicgames.launcher://ue/marketplace/content/433d28bdfb764ec5b2e60c4f45d5cb5c",
+			"SupportedTargetPlatforms": [
+				"Win64",
+				"Linux"
+			]
+		},
+		{
+			"Name": "AjaMedia",
+			"Enabled": true,
+			"MarketplaceURL": "com.epicgames.launcher://ue/marketplace/content/91088fb17fb04b6b9eda41aaf8fda21e",
+			"SupportedTargetPlatforms": [
+				"Win64"
+			]
+		},
+		{
+			"Name": "StormSyncAvalancheBridge",
+			"Enabled": true
+		},
+		{
+			"Name": "WebSocketMessaging",
+			"Enabled": true
+		},
+		{
+			"Name": "StormSync",
+			"Enabled": true
+		}
+	]
+}
diff --git a/Engine/Plugins/VirtualProduction/DynamicMaterial/Config/BaseDynamicMaterial.ini b/Engine/Plugins/VirtualProduction/DynamicMaterial/Config/BaseDynamicMaterial.ini
new file mode 100644
--- /dev/null
+++ b/Engine/Plugins/VirtualProduction/DynamicMaterial/Config/BaseDynamicMaterial.ini
@@ -0,0 +1,43 @@
+[CoreRedirects]
+
++ClassRedirects=(OldName="/Script/DynamicMaterial.DMMaterialValueTexture2D",NewName="/Script/DynamicMaterial.DMMaterialValueTexture")
++ClassRedirects=(OldName="/Script/DynamicMaterial.DMMaterialValueTextureCube",NewName="/Script/DynamicMaterial.DMMaterialValueTexture")
++ClassRedirects=(OldName="/Script/DynamicMaterial.DMMaterialValueVolumeTexture",NewName="/Script/DynamicMaterial.DMMaterialValueTexture")
++ClassRedirects=(OldName="/Script/DynamicMaterial.DynamicMaterialModelBase",NewName="/Script/DynamicMaterial.DynamicMaterialModel")
++ClassRedirects=(OldName="/Script/DynamicMaterialEditor.DynamicMaterialModel",NewName="/Script/DynamicMaterial.DynamicMaterialModel")
+
++EnumRedirects=(OldName="/Script/DynamicMaterial.EDMValueType",NewName="/Script/DynamicMaterial.EDMValueType",ValueChanges=(("VT_Texture_Any","VT_Texture"),("VT_Texture2D_Object","VT_Texture"),("VT_TextureCube_Object","VT_Texture"),("VT_VolumeTexture_Object","VT_Texture")))
+
++PackageRedirects=(OldName="/DynamicMaterial/MaterialFunctions/Effects/MF_DM_Effect_HSV.MF_DM_Effect_HSV",NewName="/DynamicMaterial/MaterialFunctions/Effects/Color/MF_DM_Effect_Base_HSV.MF_DM_Effect_Base_HSV")
++PackageRedirects=(OldName="/DynamicMaterial/MaterialFunctions/Effects/MF_DM_Effect_Base_HSV.MF_DM_Effect_Base_HSV",NewName="/DynamicMaterial/MaterialFunctions/Effects/Color/MF_DM_Effect_Base_HSV.MF_DM_Effect_Base_HSV")
++PackageRedirects=(OldName="/DynamicMaterial/MaterialFunctions/Effects/MF_DM_UVPanner.MF_DM_UVPanner",NewName="/DynamicMaterial/MaterialFunctions/Effects/UV/MF_DM_Effect_UV_Panner.MF_DM_Effect_UV_Panner")
++PackageRedirects=(OldName="/DynamicMaterial/MaterialFunctions/Effects/MF_DM_Effect_UV_Panner.MF_DM_Effect_UV_Panner",NewName="/DynamicMaterial/MaterialFunctions/Effects/UV/MF_DM_Effect_UV_Panner.MF_DM_Effect_UV_Panner")
++PackageRedirects=(OldName="/DynamicMaterial/MaterialFunctions/Effects/MF_DM_Effect_Base_Levels.MF_DM_Effect_Base_Levels",NewName="/DynamicMaterial/MaterialFunctions/Effects/Color/MF_DM_Effect_Base_Levels.MF_DM_Effect_Base_Levels")
++PackageRedirects=(OldName="/DynamicMaterial/MaterialFunctions/Effects/MF_DM_Effect_Mask_Level.MF_DM_Effect_Mask_Level",NewName="/DynamicMaterial/MaterialFunctions/Effects/Alpha/MF_DM_Effect_Mask_Level.MF_DM_Effect_Mask_Level")
++PackageRedirects=(OldName="/DynamicMaterial/MaterialFunctions/Effects/MF_DM_UV_Distort.MF_DM_UV_Distort",NewName="/DynamicMaterial/MaterialFunctions/Effects/UV/MF_DM_UV_Distort.MF_DM_UV_Distort")
++PackageRedirects=(OldName="/DynamicMaterial/MaterialFunctions/Effects/MF_DM_Effect_UV_Pixelate.MF_DM_Effect_UV_Pixelate",NewName="/DynamicMaterial/MaterialFunctions/Effects/UV/MF_DM_Effect_UV_Pixelate.MF_DM_Effect_UV_Pixelate")
+
++PropertyRedirects=(OldName="/Script/DynamicMaterial.DMRenderTargetTextRenderer.TextScale",NewName="/Script/DynamicMaterial.DMRenderTargetTextRenderer.RenderScale")
++PackageRedirects=(OldName="/Script/DynamicMaterial.DynamicMaterialModel.GlobalOpacity",NewName="/Script/DynamicMaterial.DynamicMaterialModel.GlobalOpacityValue")
++PackageRedirects=(OldName="/Script/DynamicMaterial.DynamicMaterialModel.GlobalTiling",NewName="/Script/DynamicMaterial.DynamicMaterialModel.GlobalTilingValue")
++PackageRedirects=(OldName="/Script/DynamicMaterial.DynamicMaterialModel.GlobalMetallic",NewName="/Script/DynamicMaterial.DynamicMaterialModel.GlobalMetallicParameter")
++PackageRedirects=(OldName="/Script/DynamicMaterial.DynamicMaterialModel.GlobalSpecular",NewName="/Script/DynamicMaterial.DynamicMaterialModel.GlobalSpecularParameter")
++PackageRedirects=(OldName="/Script/DynamicMaterial.DynamicMaterialModel.GlobalRoughness",NewName="/Script/DynamicMaterial.DynamicMaterialModel.GlobalRoughnessParameter")
++PackageRedirects=(OldName="/Script/DynamicMaterial.DynamicMaterialModel.GlobalAnisotropy",NewName="/Script/DynamicMaterial.DynamicMaterialModel.GlobalAnisotropyParameter")
++PackageRedirects=(OldName="/Script/DynamicMaterial.DynamicMaterialModel.GlobalWorldPositionOffset",NewName="/Script/DynamicMaterial.DynamicMaterialModel.GlobalWorldPositionOffsetParameter")
++PackageRedirects=(OldName="/Script/DynamicMaterial.DynamicMaterialModel.GlobalAmbientOcclusion",NewName="/Script/DynamicMaterial.DynamicMaterialModel.GlobalAmbientOcclusionParameter")
++PackageRedirects=(OldName="/Script/DynamicMaterial.DynamicMaterialModel.GlobalRefraction",NewName="/Script/DynamicMaterial.DynamicMaterialModel.GlobalRefractionParameter")
++PackageRedirects=(OldName="/Script/DynamicMaterial.DynamicMaterialModel.GlobalPixelDepthOffset",NewName="/Script/DynamicMaterial.DynamicMaterialModel.GlobalPixelDepthOffsetParameter")
++PackageRedirects=(OldName="/Script/DynamicMaterial.DynamicMaterialModel.GlobalOffset",NewName="/Script/DynamicMaterial.DynamicMaterialModel.GlobalOffsetParameter")
++PackageRedirects=(OldName="/Script/DynamicMaterial.DynamicMaterialModel.GlobalRotation",NewName="/Script/DynamicMaterial.DynamicMaterialModel.GlobalRotationParameter")
+
++PropertyRedirects=(OldName="/Script/DynamicMaterial.DMDefaultMaterialPropertySlotValue.Type",NewName="/Script/DynamicMaterial.DMDefaultMaterialPropertySlotValue.DefaultType")
++PropertyRedirects=(OldName="/Script/DynamicMaterial.DynamicMaterialInstance.MaterialModel",NewName="/Script/DynamicMaterial.DynamicMaterialInstance.MaterialModelBase")
++FunctionRedirects=(OldName="/Script/DynamicMaterial.DynamicMaterialModel.RemoveValueByName",NewName="/Script/DynamicMaterial.DynamicMaterialModel.RemoveValueByParameterName")
+
++PropertyRedirects=(OldName="/Script/DynamicMaterialEditor.DynamicMaterialModelEditorOnlyData.bTwoSidedFlag",NewName="/Script/DynamicMaterialEditor.DynamicMaterialModelEditorOnlyData.bTwoSided")
++PropertyRedirects=(OldName="/Script/DynamicMaterialEditor.DynamicMaterialModelEditorOnlyData.bPixelAnimationFlag",NewName="/Script/DynamicMaterialEditor.DynamicMaterialModelEditorOnlyData.bHasPixelAnimation")
++FunctionRedirects=(OldName="/Script/DynamicMaterialEditor.DynamicMaterialModelEditorOnlyData.IsPixelAnimationFlagSet",NewName="/Script/DynamicMaterialEditor.DynamicMaterialModelEditorOnlyData.GetHasPixelAnimation")
++FunctionRedirects=(OldName="/Script/DynamicMaterialEditor.DynamicMaterialModelEditorOnlyData.SetPixelAnimationFlag",NewName="/Script/DynamicMaterialEditor.DynamicMaterialModelEditorOnlyData.SetHasPixelAnimation")
++FunctionRedirects=(OldName="/Script/DynamicMaterialEditor.DynamicMaterialModelEditorOnlyData.IsTwoSidedFlagSet",NewName="/Script/DynamicMaterialEditor.DynamicMaterialModelEditorOnlyData.GetIsTwoSided")
++FunctionRedirects=(OldName="/Script/DynamicMateThatrialEditor.DynamicMaterialModelEditorOnlyData.SetTwoSidedFlag",NewName="/Script/DynamicMaterialEditor.DynamicMaterialModelEditorOnlyData.SetIsTwoSided")
\ No newline at end of file
diff --git a/Engine/Plugins/VirtualProduction/DynamicMaterial/Source/DynamicMaterialEditor/Private/Utils/DMMaterialStageFunctionLibrary.cpp b/Engine/Plugins/VirtualProduction/DynamicMaterial/Source/DynamicMaterialEditor/Private/Utils/DMMaterialStageFunctionLibrary.cpp
new file mode 100644
--- /dev/null
+++ b/Engine/Plugins/VirtualProduction/DynamicMaterial/Source/DynamicMaterialEditor/Private/Utils/DMMaterialStageFunctionLibrary.cpp
@@ -0,0 +1,96 @@
+// Copyright Epic Games, Inc. All Rights Reserved.
+
+#include "Utils/DMMaterialStageFunctionLibrary.h"
+
+#include "Components/DMMaterialStage.h"
+#include "Components/DMMaterialStageThroughput.h"
+#include "Components/DMMaterialSubStage.h"
+#include "Components/DMRenderTargetRenderer.h"
+#include "Components/MaterialStageExpressions/DMMSETextureSample.h"
+#include "Components/MaterialStageInputs/DMMSIExpression.h"
+#include "Components/MaterialStageInputs/DMMSIValue.h"
+#include "Components/MaterialValues/DMMaterialValueRenderTarget.h"
+#include "Components/MaterialValues/DMMaterialValueTexture.h"
+
+UDMMaterialStageInputValue* UDMMaterialStageFunctionLibrary::FindDefaultStageOpacityInputValue(UDMMaterialStage* InStage)
+{
+	if (!IsValid(InStage))
+	{
+		return nullptr;
+	}
+
+	int32 OpacityInputIndex = -1;
+
+	if (UDMMaterialStageThroughput* const ThroughputSource = Cast<UDMMaterialStageThroughput>(InStage->GetSource()))
+	{
+		const TArray<FDMMaterialStageConnector>& Connectors = ThroughputSource->GetInputConnectors();
+		for (const FDMMaterialStageConnector& Connector : Connectors)
+		{
+			if (Connector.Name.ToString() == TEXT("Opacity"))
+			{
+				OpacityInputIndex = Connector.Index;
+			}
+		}
+	}
+
+	const TArray<UDMMaterialStageInput*>& StageInputs = InStage->GetInputs();
+	const TArray<FDMMaterialStageConnection>& InputConnectionMap = InStage->GetInputConnectionMap();
+
+	if (InputConnectionMap.IsValidIndex(OpacityInputIndex) && !InputConnectionMap[OpacityInputIndex].Channels.IsEmpty())
+	{
+		if (InputConnectionMap[OpacityInputIndex].Channels[0].SourceIndex >= FDMMaterialStageConnectorChannel::FIRST_STAGE_INPUT)
+		{
+			const int32 StageInputIndex = InputConnectionMap[OpacityInputIndex].Channels[0].SourceIndex - FDMMaterialStageConnectorChannel::FIRST_STAGE_INPUT;
+
+			return Cast<UDMMaterialStageInputValue>(StageInputs[StageInputIndex]);
+		}
+	}
+
+	return nullptr;
+}
+
+void UDMMaterialStageFunctionLibrary::SetStageInputToRenderer(UDMMaterialStage* InStage, TSubclassOf<UDMRenderTargetRenderer> InRendererClass, int32 InInputIndex)
+{
+	UDMMaterialStageInputExpression* InputExpression = UDMMaterialStageInputExpression::ChangeStageInput_Expression(
+		InStage,
+		UDMMaterialStageExpressionTextureSample::StaticClass(),
+		InInputIndex,
+		FDMMaterialStageConnectorChannel::WHOLE_CHANNEL,
+		0,
+		FDMMaterialStageConnectorChannel::THREE_CHANNELS
+	);
+
+	if (!ensure(InputExpression))
+	{
+		return;
+	}
+
+	UDMMaterialSubStage* SubStage = InputExpression->GetSubStage();
+
+	if (!ensure(SubStage))
+	{
+		return;
+	}
+
+	UDMMaterialStageInputValue* InputValue = UDMMaterialStageInputValue::ChangeStageInput_NewLocalValue(
+		SubStage,
+		0,
+		FDMMaterialStageConnectorChannel::WHOLE_CHANNEL,
+		UDMMaterialValueRenderTarget::StaticClass(),
+		FDMMaterialStageConnectorChannel::WHOLE_CHANNEL
+	);
+
+	if (!ensure(InputValue))
+	{
+		return;
+	}
+
+	UDMMaterialValueRenderTarget* RenderTargetValue = Cast<UDMMaterialValueRenderTarget>(InputValue->GetValue());
+
+	if (!RenderTargetValue)
+	{
+		return;
+	}
+
+	UDMRenderTargetRenderer::CreateRenderTargetRenderer(InRendererClass, RenderTargetValue);
+}
diff --git a/Engine/Plugins/VirtualProduction/PropertyAnimator/Config/BasePropertyAnimator.ini b/Engine/Plugins/VirtualProduction/PropertyAnimator/Config/BasePropertyAnimator.ini
new file mode 100644
--- /dev/null
+++ b/Engine/Plugins/VirtualProduction/PropertyAnimator/Config/BasePropertyAnimator.ini
@@ -0,0 +1,32 @@
+﻿[CoreRedirects]
++ClassRedirects=(OldName="/Script/PropertyController.PropertyControllerWiggle",NewName="/Script/PropertyAnimator.PropertyAnimatorWiggle")
++ClassRedirects=(OldName="/Script/PropertyController.PropertyControllerWave",NewName="/Script/PropertyAnimator.PropertyAnimatorWave")
++ClassRedirects=(OldName="/Script/PropertyController.PropertyControllerTime",NewName="/Script/PropertyAnimator.PropertyAnimatorTime")
++ClassRedirects=(OldName="/Script/PropertyController.PropertyControllerOscillate",NewName="/Script/PropertyAnimator.PropertyAnimatorOscillate")
++ClassRedirects=(OldName="/Script/PropertyController.PropertyControllerFloatBase",NewName="/Script/PropertyAnimator.PropertyAnimatorNumericBase")
++ClassRedirects=(OldName="/Script/PropertyController.PropertyControllerBounce",NewName="/Script/PropertyAnimator.PropertyAnimatorBounce")
++ClassRedirects=(OldName="/Script/PropertyController.PropertyControllerBlink",NewName="/Script/PropertyAnimator.PropertyAnimatorPulse")
++ClassRedirects=(OldName="/Script/PropertyAnimator.PropertyAnimatorBlink",NewName="/Script/PropertyAnimator.PropertyAnimatorPulse")
++ClassRedirects=(OldName="/Script/PropertyController.PropertyControllerBase",NewName="/Script/PropertyAnimator.PropertyAnimatorBase")
++ClassRedirects=(OldName="/Script/PropertyController.PropertyControllerWaveDoubleChannelOverrideContainer",NewName="/Script/PropertyAnimator.PropertyAnimatorWaveDoubleChannelOverrideContainer")
++ClassRedirects=(OldName="/Script/PropertyController.PropertyControllerWaveDoubleChannelEvaluatorSystem",NewName="/Script/PropertyAnimator.PropertyAnimatorWaveDoubleChannelEvaluatorSystem")
++ClassRedirects=(OldName="/Script/PropertyController.PropertyControllerEasingDoubleChannelOverrideContainer",NewName="/Script/PropertyAnimator.PropertyAnimatorEasingDoubleChannelOverrideContainer")
++ClassRedirects=(OldName="/Script/PropertyController.PropertyControllerEasingDoubleChannelEvaluatorSystem",NewName="/Script/PropertyAnimator.PropertyAnimatorEasingDoubleChannelEvaluatorSystem")
++ClassRedirects=(OldName="/Script/PropertyController.PropertyControllerCustomTimeSource",NewName="/Script/PropertyAnimator.PropertyAnimatorCustomTimeSource")
++ClassRedirects=(OldName="/Script/PropertyController.PropertyControllerTimeSourceBase",NewName="/Script/PropertyAnimator.PropertyAnimatorTimeSourceBase")
++ClassRedirects=(OldName="/Script/PropertyController.PropertyControllerWorldTimeSource",NewName="/Script/PropertyAnimator.PropertyAnimatorWorldTimeSource")
++ClassRedirects=(OldName="/Script/AvalanchePropertyAnimator.AvaPropertyAnimatorTextResolver",NewName="/Script/PropertyAnimator.PropertyAnimatorTextResolver")
++ClassRedirects=(OldName="/Script/PropertyAnimator.PropertyAnimatorCurveWave",NewName="/Script/PropertyAnimator.PropertyAnimatorCurve")
++ClassRedirects=(OldName="/Script/PropertyAnimator.PropertyAnimatorFloatBase",NewName="/Script/PropertyAnimator.PropertyAnimatorNumericBase")
+
++EnumRedirects=(OldName="/Script/PropertyController.EPropertyControllerEasingType",NewName="/Script/PropertyAnimator.EPropertyAnimatorEasingType")
++EnumRedirects=(OldName="/Script/PropertyController.EPropertyControllerEasingFunction",NewName="/Script/PropertyAnimator.EPropertyAnimatorEasingFunction")
++EnumRedirects=(OldName="/Script/PropertyController.EPropertyControllerWaveFunction",NewName="/Script/PropertyAnimator.EPropertyAnimatorWaveFunction")
+
++StructRedirects=(OldName="/Script/PropertyController.PropertyControllerFloatPropertyOptions",NewName="/Script/PropertyAnimator.PropertyAnimatorFloatPropertyOptions")
++StructRedirects=(OldName="/Script/PropertyController.PropertyControllerEasingDoubleChannel",NewName="/Script/PropertyAnimator.PropertyAnimatorEasingDoubleChannel")
++StructRedirects=(OldName="/Script/PropertyController.PropertyControllerEasingParameters",NewName="/Script/PropertyAnimator.PropertyAnimatorEasingParameters")
++StructRedirects=(OldName="/Script/PropertyController.PropertyControllerWaveDoubleChannel",NewName="/Script/PropertyAnimator.PropertyAnimatorWaveDoubleChannel")
++StructRedirects=(OldName="/Script/PropertyController.PropertyControllerWaveParameters",NewName="/Script/PropertyAnimator.PropertyAnimatorWaveParameters")
+
++PropertyRedirects=(OldName="/Script/PropertyAnimator.PropertyAnimatorFloatBase.Strength",NewName="/Script/PropertyAnimator.PropertyAnimatorNumericBase.GlobalMagnitude")
\ No newline at end of file
diff --git a/Engine/Plugins/VirtualProduction/PropertyAnimator/Source/PropertyAnimator/Private/PropertyAnimatorModule.cpp b/Engine/Plugins/VirtualProduction/PropertyAnimator/Source/PropertyAnimator/Private/PropertyAnimatorModule.cpp
new file mode 100644
--- /dev/null
+++ b/Engine/Plugins/VirtualProduction/PropertyAnimator/Source/PropertyAnimator/Private/PropertyAnimatorModule.cpp
@@ -0,0 +1,5 @@
+﻿// Copyright Epic Games, Inc. All Rights Reserved.
+
+#include "Modules/ModuleManager.h"
+
+IMPLEMENT_MODULE(FDefaultModuleImpl, PropertyAnimator)
diff --git a/Engine/Plugins/VirtualProduction/PropertyAnimatorCore/Config/BasePropertyAnimatorCore.ini b/Engine/Plugins/VirtualProduction/PropertyAnimatorCore/Config/BasePropertyAnimatorCore.ini
new file mode 100644
--- /dev/null
+++ b/Engine/Plugins/VirtualProduction/PropertyAnimatorCore/Config/BasePropertyAnimatorCore.ini
@@ -0,0 +1,24 @@
+﻿[CoreRedirects]
++ClassRedirects=(OldName="/Script/CustomPropertyControl.CustomPropertyControlResolver",NewName="/Script/PropertyAnimatorCore.PropertyAnimatorCoreResolver")
++ClassRedirects=(OldName="/Script/CustomPropertyControl.CustomPropertyControlHandlerBase",NewName="/Script/PropertyAnimatorCore.PropertyAnimatorCoreHandlerBase")
++ClassRedirects=(OldName="/Script/CustomPropertyControl.CustomPropertyControlFloatHandler",NewName="/Script/PropertyAnimatorCore.PropertyAnimatorCoreFloatHandler")
++ClassRedirects=(OldName="/Script/CustomPropertyControl.CustomPropertyControlDoubleHandler",NewName="/Script/PropertyAnimatorCore.PropertyAnimatorCoreDoubleHandler")
++ClassRedirects=(OldName="/Script/CustomPropertyControl.CustomPropertyControlControllerBase",NewName="/Script/PropertyAnimatorCore.PropertyAnimatorCoreBase")
++ClassRedirects=(OldName="/Script/CustomPropertyControl.CustomPropertyControlComponent",NewName="/Script/PropertyAnimatorCore.PropertyAnimatorCoreComponent")
++ClassRedirects=(OldName="/Script/PropertyController.PropertyControllerCustomTimeSource",NewName="/Script/PropertyAnimator.PropertyAnimatorCustomTimeSource")
++ClassRedirects=(OldName="/Script/PropertyController.PropertyControllerTimeSourceBase",NewName="/Script/PropertyAnimator.PropertyAnimatorTimeSourceBase")
++ClassRedirects=(OldName="/Script/PropertyController.PropertyControllerWorldTimeSource",NewName="/Script/PropertyAnimator.PropertyAnimatorWorldTimeSource")
++ClassRedirects=(OldName="/Script/PropertyAnimator.PropertyAnimatorCustomTimeSource",NewName="/Script/PropertyAnimatorCore.PropertyAnimatorCoreManualTimeSource")
++ClassRedirects=(OldName="/Script/PropertyAnimator.PropertyAnimatorTimeSourceBase",NewName="/Script/PropertyAnimatorCore.PropertyAnimatorCoreTimeSourceBase")
++ClassRedirects=(OldName="/Script/PropertyAnimator.PropertyAnimatorWorldTimeSource",NewName="/Script/PropertyAnimatorCore.PropertyAnimatorCoreWorldTimeSource")
++ClassRedirects=(OldName="/Script/AvalanchePropertyAnimator.AvaPropertyAnimatorSequenceTimeSource",NewName="/Script/PropertyAnimatorCore.PropertyAnimatorCoreSequencerTimeSource")
++ClassRedirects=(OldName="/Script/PropertyAnimatorCore.PropertyAnimatorCoreMachineClockTimeSource",NewName="/Script/PropertyAnimatorCore.PropertyAnimatorCoreSystemTimeSource")
+
++StructRedirects=(OldName="/Script/CustomPropertyControlEditor.CustomPropertyControlData",NewName="/Script/PropertyAnimatorCore.PropertyAnimatorCoreData")
+
++PropertyRedirects=(OldName="/Script/PropertyAnimatorCore.PropertyAnimatorCoreBase.ControllerDisplayName",NewName="/Script/PropertyAnimatorCore.PropertyAnimatorCoreBase.AnimatorDisplayName")
++PropertyRedirects=(OldName="/Script/PropertyAnimatorCore.PropertyAnimatorCoreComponent.Controllers",NewName="/Script/PropertyAnimatorCore.PropertyAnimatorCoreComponent.Animators")
++PropertyRedirects=(OldName="/Script/PropertyAnimatorCore.PropertyAnimatorCoreBase.bEnabled",NewName="/Script/PropertyAnimatorCore.PropertyAnimatorCoreBase.bAnimatorEnabled")
++PropertyRedirects=(OldName="/Script/PropertyAnimatorCore.PropertyAnimatorCoreComponent.bEnabled",NewName="/Script/PropertyAnimatorCore.PropertyAnimatorCoreComponent.bAnimatorsEnabled")
+
++EnumRedirects=(OldName="/Script/PropertyAnimatorCore.EPropertyAnimatorCoreMachineClockMode",NewName="/Script/PropertyAnimatorCore.EPropertyAnimatorCoreSystemMode")
\ No newline at end of file
diff --git a/Engine/Plugins/VirtualProduction/PropertyAnimatorCore/Source/PropertyAnimatorCore/Private/Properties/Handlers/PropertyAnimatorCoreTextHandler.cpp b/Engine/Plugins/VirtualProduction/PropertyAnimatorCore/Source/PropertyAnimatorCore/Private/Properties/Handlers/PropertyAnimatorCoreTextHandler.cpp
new file mode 100644
--- /dev/null
+++ b/Engine/Plugins/VirtualProduction/PropertyAnimatorCore/Source/PropertyAnimatorCore/Private/Properties/Handlers/PropertyAnimatorCoreTextHandler.cpp
@@ -0,0 +1,101 @@
+// Copyright Epic Games, Inc. All Rights Reserved.
+
+#include "Properties/Handlers/PropertyAnimatorCoreTextHandler.h"
+#include "StructUtils/PropertyBag.h"
+#include "UObject/TextProperty.h"
+
+bool UPropertyAnimatorCoreTextHandler::IsPropertySupported(const FPropertyAnimatorCoreData& InPropertyData) const
+{
+	if (InPropertyData.IsA<FTextProperty>())
+	{
+		return true;
+	}
+
+	return Super::IsPropertySupported(InPropertyData);
+}
+
+bool UPropertyAnimatorCoreTextHandler::GetValue(const FPropertyAnimatorCoreData& InPropertyData, FInstancedPropertyBag& OutValue)
+{
+	const FName PropertyHash(InPropertyData.GetLocatorPathHash());
+	OutValue.AddProperty(PropertyHash, EPropertyBagPropertyType::Text);
+
+	FText Value;
+	InPropertyData.GetPropertyValuePtr(&Value);
+
+	OutValue.SetValueText(PropertyHash, Value);
+
+	return true;
+}
+
+bool UPropertyAnimatorCoreTextHandler::SetValue(const FPropertyAnimatorCoreData& InPropertyData, const FInstancedPropertyBag& InValue)
+{
+	const FName PropertyHash(InPropertyData.GetLocatorPathHash());
+	TValueOrError<FText, EPropertyBagResult> ValueResult = InValue.GetValueText(PropertyHash);
+	if (!ValueResult.HasValue())
+	{
+		return false;
+	}
+
+	FText& NewValue = ValueResult.GetValue();
+	InPropertyData.SetPropertyValuePtr(&NewValue);
+
+	return true;
+}
+
+bool UPropertyAnimatorCoreTextHandler::IsAdditiveSupported() const
+{
+	return true;
+}
+
+bool UPropertyAnimatorCoreTextHandler::AddValue(const FPropertyAnimatorCoreData& InPropertyData, const FInstancedPropertyBag& InValue)
+{
+	const FName PropertyHash(InPropertyData.GetLocatorPathHash());
+	const TValueOrError<FText, EPropertyBagResult> ValueResult = InValue.GetValueText(PropertyHash);
+	if (!ValueResult.HasValue())
+	{
+		return false;
+	}
+
+	FText Value;
+	InPropertyData.GetPropertyValuePtr(&Value);
+
+	// Append
+	FText NewValue = FText::FromString(Value.ToString() + ValueResult.GetValue().ToString());
+	InPropertyData.SetPropertyValuePtr(&NewValue);
+
+	return true;
+}
+
+bool UPropertyAnimatorCoreTextHandler::SubtractValue(const FPropertyAnimatorCoreData& InPropertyData, const FInstancedPropertyBag& InValue)
+{
+	const FName PropertyHash(InPropertyData.GetLocatorPathHash());
+	const TValueOrError<FText, EPropertyBagResult> ValueResult = InValue.GetValueText(PropertyHash);
+	if (!ValueResult.HasValue())
+	{
+		return false;
+	}
+
+	FText Value;
+	InPropertyData.GetPropertyValuePtr(&Value);
+
+	// Trim end
+	FString StringValue = Value.ToString();
+	const FString StringValueResult = ValueResult.GetValue().ToString();
+	if (StringValue.EndsWith(StringValueResult))
+	{
+		StringValue.LeftChopInline(StringValueResult.Len());
+	}
+
+	FText NewValue = FText::FromString(StringValue);
+	InPropertyData.SetPropertyValuePtr(&NewValue);
+
+	return true;
+}
+
+bool UPropertyAnimatorCoreTextHandler::GetDefaultValue(const FPropertyAnimatorCoreData& InPropertyData, FInstancedPropertyBag& OutValue)
+{
+	const FName PropertyHash(InPropertyData.GetLocatorPathHash());
+	OutValue.AddProperty(PropertyHash, EPropertyBagPropertyType::Text);
+	OutValue.SetValueText(PropertyHash, FText::GetEmpty());
+	return true;
+}
diff --git a/Engine/Plugins/VirtualProduction/StormSync/Source/StormSyncDrives/Private/StormSyncDrivesModule.cpp b/Engine/Plugins/VirtualProduction/StormSync/Source/StormSyncDrives/Private/StormSyncDrivesModule.cpp
new file mode 100644
--- /dev/null
+++ b/Engine/Plugins/VirtualProduction/StormSync/Source/StormSyncDrives/Private/StormSyncDrivesModule.cpp
@@ -0,0 +1,283 @@
+// Copyright Epic Games, Inc. All Rights Reserved.
+
+#include "StormSyncDrivesModule.h"
+
+#include "Engine/DeveloperSettings.h"
+#include "Misc/CoreDelegates.h"
+#include "Misc/PackageName.h"
+#include "StormSyncDrivesLog.h"
+#include "StormSyncDrivesSettings.h"
+#include "StormSyncDrivesUtils.h"
+
+#if WITH_EDITOR
+#include "MessageLogModule.h"
+#include "Logging/MessageLog.h"
+#endif
+
+#define LOCTEXT_NAMESPACE "FStormSyncDrivesModule"
+
+void FStormSyncDrivesModule::StartupModule()
+{
+#if WITH_EDITOR
+	// Create a message log for the validation to use
+	FMessageLogModule& MessageLogModule = FModuleManager::LoadModuleChecked<FMessageLogModule>("MessageLog");
+	FMessageLogInitializationOptions InitOptions;
+	InitOptions.bShowPages = true;
+	MessageLogModule.RegisterLogListing(LogName, LOCTEXT("StormSyncLogLabel", "Storm Sync Drives"), InitOptions);
+	LogListing = MessageLogModule.GetLogListing(LogName);
+#endif
+
+	FCoreDelegates::OnPostEngineInit.AddRaw(this, &FStormSyncDrivesModule::OnPostEngineInit);
+
+#if WITH_EDITOR
+	GetMutableDefault<UStormSyncDrivesSettings>()->OnSettingChanged().AddRaw(this, &FStormSyncDrivesModule::OnSettingsChanged);
+#endif
+}
+
+void FStormSyncDrivesModule::ShutdownModule()
+{
+#if WITH_EDITOR
+	if (FModuleManager::Get().IsModuleLoaded("MessageLog"))
+	{
+		// unregister message log
+		FMessageLogModule& MessageLogModule = FModuleManager::GetModuleChecked<FMessageLogModule>("MessageLog");
+		MessageLogModule.UnregisterLogListing(LogName);
+	}
+#endif
+
+	FCoreDelegates::OnPostEngineInit.RemoveAll(this);
+
+#if WITH_EDITOR
+	if (UObjectInitialized())
+	{
+		GetMutableDefault<UStormSyncDrivesSettings>()->OnSettingChanged().RemoveAll(this);
+	}
+#endif
+
+	MountedDrives.Empty();
+}
+
+bool FStormSyncDrivesModule::RegisterMountPoint(const FStormSyncMountPointConfig& InMountPoint, FText& ErrorText, EStormSyncDriveErrorCode* OutErrorCode)
+{
+	UE_LOG(LogStormSyncDrives, Display, TEXT("FStormSyncDrivesModule::RegisterMountPoint ... MountPoint: %s, Path: %s"), *InMountPoint.MountPoint, *InMountPoint.MountDirectory.Path);
+
+	// Validate first configuration of the mount point
+	FText ValidationError;
+	if (!FStormSyncDrivesUtils::ValidateMountPoint(InMountPoint, ValidationError, OutErrorCode))
+	{
+		ErrorText = FText::Format(
+			LOCTEXT("ValidationError_MountPoint", "Validation failed for MountPoint (MountPoint: {0}, MountDirectory: {1}) - {2}"),
+			FText::FromString(InMountPoint.MountPoint),
+			FText::FromString(InMountPoint.MountDirectory.Path),
+			ValidationError
+		);
+		AddMessageError(ErrorText);
+		return false;
+	}
+
+	// Append a trailing slash so that it is consistent with FLongPackagePathsSingleton::ContentPathToRoot entries and avoid a crash
+	// with assets and reimport path when right clicking on them
+
+	// ValidateMountPoint should ensure no trailing slash is allowed in user config
+	const FString MountPoint = InMountPoint.MountPoint + TEXT("/");
+
+	// Check if root path is already mounted
+	if (FPackageName::MountPointExists(MountPoint))
+	{
+		ErrorText = FText::Format(
+			LOCTEXT("ValidationError_MountPointExists", "Mount Point {0} already exists"),
+			FText::FromString(MountPoint)
+		);
+		AddMessageError(ErrorText);
+		return false;
+	}
+
+	UE_LOG(LogStormSyncDrives, Display, TEXT("FStormSyncDrivesModule::RegisterMountPoint ... Try register %s to %s"), *MountPoint, *InMountPoint.MountDirectory.Path);
+	FPackageName::RegisterMountPoint(MountPoint, InMountPoint.MountDirectory.Path);
+	return true;
+}
+
+bool FStormSyncDrivesModule::UnregisterMountPoint(const FStormSyncMountPointConfig& InMountPoint, FText& ErrorText)
+{
+	const FString ContentPath = InMountPoint.MountDirectory.Path;
+
+	// Mount points are always added with a trailing slash (eg. /Foo => /Foo/)
+	const bool bHasLeadingSlash = InMountPoint.MountPoint.EndsWith(TEXT("/"));
+	const FString RootPath = bHasLeadingSlash ? InMountPoint.MountPoint : InMountPoint.MountPoint + TEXT("/");
+
+	if (!FPackageName::MountPointExists(RootPath))
+	{
+		ErrorText = FText::Format(
+			LOCTEXT("ValidationError_MountPointUnexisting", "Mount Point {0} does not exist"),
+			FText::FromString(RootPath)
+		);
+		return false;
+	}
+
+	FPackageName::UnRegisterMountPoint(RootPath, ContentPath);
+	return true;
+}
+
+void FStormSyncDrivesModule::OnPostEngineInit()
+{
+	UE_LOG(LogStormSyncDrives, Display, TEXT("FStormSyncDrivesModule::OnPostEngineInit ..."));
+	UE_LOG(LogStormSyncDrives, Display, TEXT("\t Mounting Drives based on config"));
+
+	ResetMountedDrivesFromSettings(GetDefault<UStormSyncDrivesSettings>());
+}
+
+void FStormSyncDrivesModule::OnSettingsChanged(UObject* InSettings, FPropertyChangedEvent& InPropertyChangedEvent)
+{
+	const FName PropertyName = InPropertyChangedEvent.GetPropertyName();
+	const FName MemberPropertyName = (InPropertyChangedEvent.MemberProperty != nullptr) ? InPropertyChangedEvent.MemberProperty->GetFName() : NAME_None;
+	const FName MountPointsPropertyName = GET_MEMBER_NAME_CHECKED(UStormSyncDrivesSettings, MountPoints);
+
+
+	if (PropertyName == MountPointsPropertyName || MemberPropertyName == MountPointsPropertyName)
+	{
+		ResetMountedDrivesFromSettings(CastChecked<UStormSyncDrivesSettings>(InSettings));
+	}
+}
+
+bool FStormSyncDrivesModule::ValidateMountPoint(const FStormSyncMountPointConfig& InMountPoint, int32 Index) const
+{
+	const FString RootPath = InMountPoint.MountPoint;
+	const FString Directory = InMountPoint.MountDirectory.Path;
+
+	FText ErrorText;
+	if (!FStormSyncDrivesUtils::ValidateMountPoint(InMountPoint, ErrorText))
+	{
+		const FText ErrorLog = FText::Format(
+			LOCTEXT("ValidationError_MountPointIndex", "Validation failed for MountPoint at index {0} (MountPoint: {1}, MountDirectory: {2}) - {3}"),
+			FText::AsNumber(Index),
+			FText::FromString(RootPath),
+			FText::FromString(Directory),
+			ErrorText
+		);
+		AddMessageError(ErrorLog);
+		UE_LOG(LogStormSyncDrives, Error, TEXT("FStormSyncDrivesModule::ValidateMountPoint - %s"), *ErrorLog.ToString());
+		return false;
+	}
+
+	return true;
+}
+
+bool FStormSyncDrivesModule::ValidateMountPoints(const TArray<FStormSyncMountPointConfig>& InMountPoints) const
+{
+#if WITH_EDITOR
+	// Reset log listing and clear message on each validation run
+	if (LogListing.IsValid())
+	{
+		LogListing->ClearMessages();
+	}
+#endif
+
+	bool bResult = true;
+
+	// First check for mount point validity individually
+	int32 Index = 0;
+	for (const FStormSyncMountPointConfig& Entry : InMountPoints)
+	{
+		UE_LOG(LogStormSyncDrives, Display, TEXT("Try validate entry %s %s"), *Entry.MountPoint, *Entry.MountDirectory.Path);
+
+		// First check for mount point validity individually
+		bResult &= ValidateMountPoint(Entry, Index);
+		Index++;
+	}
+
+	// Then check for duplicates
+	TArray<FText> ValidationErrors;
+	if (!FStormSyncDrivesUtils::ValidateNonDuplicates(InMountPoints, ValidationErrors))
+	{
+		bResult = false;
+
+		// Log any returned errors to message log
+		for (const FText& ValidationError : ValidationErrors)
+		{
+			AddMessageError(ValidationError);
+			UE_LOG(LogStormSyncDrives, Error, TEXT("FStormSyncDrivesModule::ValidateMountPoints - %s"), *ValidationError.ToString());
+		}
+	}
+
+#if WITH_EDITOR
+	// Got some errors, notify user
+	if (!bResult)
+	{
+		LogListing->NotifyIfAnyMessages(LOCTEXT("Validation_Error", "Storm Sync: Error validating Mount Points settings"), EMessageSeverity::Error, true);
+	}
+#endif
+
+	return bResult;
+}
+
+void FStormSyncDrivesModule::ResetMountedDrivesFromSettings(const UStormSyncDrivesSettings* InSettings)
+{
+	UE_LOG(LogStormSyncDrives, Display, TEXT("FStormSyncDrivesModule::ResetMountedDrives ..."));
+	check(InSettings);
+
+	if (ValidateMountPoints(InSettings->MountPoints))
+	{
+		UnregisterMountedDrives();
+		CacheMountedDrives(InSettings->MountPoints);
+		RegisterMountedDrives();
+	}
+}
+
+void FStormSyncDrivesModule::UnregisterMountedDrives()
+{
+	UE_LOG(LogStormSyncDrives, Display, TEXT("FStormSyncDrivesModule::UnregisterMountedDrives ..."));
+	for (const TSharedPtr<FStormSyncMountPointConfig>& MountedDrive : MountedDrives)
+	{
+		if (!MountedDrive.IsValid())
+		{
+			continue;
+		}
+
+		FText ErrorText;
+		if (!UnregisterMountPoint(*MountedDrive.Get(), ErrorText))
+		{
+			UE_LOG(LogStormSyncDrives, Error, TEXT("FStormSyncDrivesModule::UnregisterMountedDrives failed with Error: %s"), *ErrorText.ToString());
+		}
+	}
+}
+
+void FStormSyncDrivesModule::CacheMountedDrives(const TArray<FStormSyncMountPointConfig>& InMountPoints)
+{
+	UE_LOG(LogStormSyncDrives, Display, TEXT("FStormSyncDrivesModule::CacheMountedDrives ..."));
+	MountedDrives.Reset();
+
+	for (const FStormSyncMountPointConfig& Entry : InMountPoints)
+	{
+		MountedDrives.Add(MakeShared<FStormSyncMountPointConfig>(Entry));
+	}
+}
+
+void FStormSyncDrivesModule::RegisterMountedDrives()
+{
+	UE_LOG(LogStormSyncDrives, Display, TEXT("FStormSyncDrivesModule::RegisterMountedDrives ..."));
+	for (TSharedPtr<FStormSyncMountPointConfig> MountedDrive : MountedDrives)
+	{
+		if (!MountedDrive.IsValid())
+		{
+			continue;
+		}
+
+		FText ErrorText;
+		if (!RegisterMountPoint(*MountedDrive.Get(), ErrorText))
+		{
+			UE_LOG(LogStormSyncDrives, Error, TEXT("FStormSyncDrivesModule::RegisterMountedDrives failed with Error: %s"), *ErrorText.ToString());
+		}
+	}
+}
+
+void FStormSyncDrivesModule::AddMessageError(const FText& Text)
+{
+#if WITH_EDITOR
+	FMessageLog MessageLog(LogName);
+	MessageLog.Error(Text);
+#endif
+}
+
+#undef LOCTEXT_NAMESPACE
+
+IMPLEMENT_MODULE(FStormSyncDrivesModule, StormSyncDrives)
diff --git a/Engine/Plugins/VirtualProduction/StormSync/Source/StormSyncImport/Private/Subsystems/StormSyncImportWorldSubsystem.cpp b/Engine/Plugins/VirtualProduction/StormSync/Source/StormSyncImport/Private/Subsystems/StormSyncImportWorldSubsystem.cpp
new file mode 100644
--- /dev/null
+++ b/Engine/Plugins/VirtualProduction/StormSync/Source/StormSyncImport/Private/Subsystems/StormSyncImportWorldSubsystem.cpp
@@ -0,0 +1,68 @@
+// Copyright Epic Games, Inc. All Rights Reserved.
+
+
+#include "Subsystems/StormSyncImportWorldSubsystem.h"
+
+#include "Engine/World.h"
+#include "StormSyncCoreDelegates.h"
+#include "StormSyncImportLog.h"
+#include "Subsystems/StormSyncImportSubsystem.h"
+#include "Tasks/StormSyncImportBufferTask.h"
+#include "Tasks/StormSyncImportFilesTask.h"
+
+void UStormSyncImportWorldSubsystem::Initialize(FSubsystemCollectionBase& Collection)
+{
+	Super::Initialize(Collection);
+	UE_LOG(LogStormSyncImport, Verbose, TEXT("UStormSyncImportWorldSubsystem::Initialize (World: %s) - %s"), *GetNameSafe(GetWorld()), *GetName());
+
+	FStormSyncCoreDelegates::OnRequestImportFile.AddUObject(this, &UStormSyncImportWorldSubsystem::HandleImportFile);
+	FStormSyncCoreDelegates::OnRequestImportBuffer.AddUObject(this, &UStormSyncImportWorldSubsystem::HandleImportBuffer);
+}
+
+void UStormSyncImportWorldSubsystem::Deinitialize()
+{
+	Super::Deinitialize();
+	UE_LOG(LogStormSyncImport, Verbose, TEXT("UStormSyncImportWorldSubsystem::Deinitialize (World: %s) - %s"), *GetNameSafe(GetWorld()), *GetName());
+
+	FStormSyncCoreDelegates::OnRequestImportFile.RemoveAll(this);
+	FStormSyncCoreDelegates::OnRequestImportBuffer.RemoveAll(this);
+}
+
+void UStormSyncImportWorldSubsystem::HandleImportFile(const FString& InFilename) const
+{
+	UE_LOG(LogStormSyncImport, Verbose, TEXT("UStormSyncImportWorldSubsystem::HandleImportFile InFilename: %s (World: %s)"), *InFilename, *GetNameSafe(GetWorld()));
+	const UWorld* World = GetWorld();
+	if (!World)
+	{
+		UE_LOG(LogStormSyncImport, Error, TEXT("UStormSyncImportWorldSubsystem::ImportNextTick failed because of invalid world (World: %s)"), *GetNameSafe(GetWorld()));
+		return;
+	}
+
+	UStormSyncImportSubsystem::Get().EnqueueImportTask(MakeShared<FStormSyncImportFilesTask>(InFilename), GetWorld());
+}
+
+void UStormSyncImportWorldSubsystem::HandleImportBuffer(const FStormSyncPackageDescriptor& InPackageDescriptor, const FStormSyncArchivePtr& InArchive) const
+{
+	if (!InArchive.IsValid())
+	{
+		UE_LOG(LogStormSyncImport, Error, TEXT("UStormSyncImportSubsystem::HandleImportBuffer failed from invalid archive"));
+		return;
+	}
+	
+	UE_LOG(LogStormSyncImport, 
+		Verbose,
+		TEXT("UStormSyncImportWorldSubsystem::HandleImportBuffer InPackageDescriptor: %s, BufferSize: %lld (World: %s)"),
+		*InPackageDescriptor.ToString(),
+		InArchive->TotalSize(),
+		*GetNameSafe(GetWorld())
+	);
+
+	const UWorld* World = GetWorld();
+	if (!World)
+	{
+		UE_LOG(LogStormSyncImport, Error, TEXT("UStormSyncImportWorldSubsystem::HandleImportBuffer failed because of invalid world (World: %s)"), *GetNameSafe(GetWorld()));
+		return;
+	}
+
+	UStormSyncImportSubsystem::Get().EnqueueImportTask(MakeShared<FStormSyncImportBufferTask>(InPackageDescriptor, InArchive), GetWorld());
+}
diff --git a/Engine/Plugins/VirtualProduction/StormSync/Source/StormSyncImport/Private/Tasks/StormSyncImportBufferTask.cpp b/Engine/Plugins/VirtualProduction/StormSync/Source/StormSyncImport/Private/Tasks/StormSyncImportBufferTask.cpp
new file mode 100644
--- /dev/null
+++ b/Engine/Plugins/VirtualProduction/StormSync/Source/StormSyncImport/Private/Tasks/StormSyncImportBufferTask.cpp
@@ -0,0 +1,18 @@
+// Copyright Epic Games, Inc. All Rights Reserved.
+
+#include "Tasks/StormSyncImportBufferTask.h"
+
+#include "StormSyncImportLog.h"
+#include "Subsystems/StormSyncImportSubsystem.h"
+
+void FStormSyncImportBufferTask::Run()
+{
+	if (!Archive.IsValid())
+	{
+		UE_LOG(LogStormSyncImport, Error, TEXT("FStormSyncImportBufferTask::Run failed on invalid archive"));
+		return;
+	}
+
+	UE_LOG(LogStormSyncImport, Display, TEXT("FStormSyncImportBufferTask::Run for buffer of size %lld"), Archive->TotalSize());
+	UStormSyncImportSubsystem::Get().PerformBufferImport(PackageDescriptor, MoveTemp(Archive));
+}
