Refactor StormSyncDrives module to allow for command line specification of mount paths as support for relative paths.

[RN] Added support specify additional mount points via the command line for the storm sync plugin.

diff --git a/Engine/Plugins/VirtualProduction/StormSync/Source/StormSyncDrives/Private/StormSyncDrivesModule.cpp b/Engine/Plugins/VirtualProduction/StormSync/Source/StormSyncDrives/Private/StormSyncDrivesModule.cpp
--- a/Engine/Plugins/VirtualProduction/StormSync/Source/StormSyncDrives/Private/StormSyncDrivesModule.cpp
+++ b/Engine/Plugins/VirtualProduction/StormSync/Source/StormSyncDrives/Private/StormSyncDrivesModule.cpp
@@ -9,6 +9,10 @@
 #include "StormSyncDrivesSettings.h"
 #include "StormSyncDrivesUtils.h"
 
+#include "Misc/CommandLine.h"
+#include "Misc/Paths.h"
+#include "Algo/AnyOf.h"
+
 #if WITH_EDITOR
 #include "MessageLogModule.h"
 #include "Logging/MessageLog.h"
@@ -47,6 +51,8 @@ void FStormSyncDrivesModule::ShutdownModule()
 
 	FCoreDelegates::OnPostEngineInit.RemoveAll(this);
 
+	UnregisterMountedDrives();
+	
 #if WITH_EDITOR
 	if (UObjectInitialized())
 	{
@@ -54,7 +60,7 @@ void FStormSyncDrivesModule::ShutdownModule()
 	}
 #endif
 
-	MountedDrives.Empty();
+	check(CachedMountPoints.IsEmpty());
 }
 
 bool FStormSyncDrivesModule::RegisterMountPoint(const FStormSyncMountPointConfig& InMountPoint, FText& ErrorText, EStormSyncDriveErrorCode* OutErrorCode)
@@ -94,6 +100,8 @@ bool FStormSyncDrivesModule::RegisterMountPoint(const FStormSyncMountPointConfig
 
 	UE_LOG(LogStormSyncDrives, Display, TEXT("FStormSyncDrivesModule::RegisterMountPoint ... Try register %s to %s"), *MountPoint, *InMountPoint.MountDirectory.Path);
 	FPackageName::RegisterMountPoint(MountPoint, InMountPoint.MountDirectory.Path);
+
+	CachedMountPoints.FindOrAdd(InMountPoint.MountPoint) = InMountPoint;
 	return true;
 }
 
@@ -115,17 +123,115 @@ bool FStormSyncDrivesModule::UnregisterMountPoint(const FStormSyncMountPointConf
 	}
 
 	FPackageName::UnRegisterMountPoint(RootPath, ContentPath);
+
+	CachedMountPoints.Remove(InMountPoint.MountPoint);
 	return true;
 }
 
+TArray<FStormSyncMountPointConfig> FStormSyncDrivesModule::GetCommandLineConfigs() const
+{
+	TArray<FStormSyncMountPointConfig> Configs;
+	FString AdditionalMountPointsString;
+	if (FParse::Value(FCommandLine::Get(), TEXT("-STORMSYNCMOUNTPOINTS"), AdditionalMountPointsString))
+	{
+		// look over a list of cvars
+		TArray<FString> AdditionalMountPoints;
+		AdditionalMountPointsString.ParseIntoArray(AdditionalMountPoints, TEXT(","), true);
+		for (const FString& AdditionalMountPoint : AdditionalMountPoints)
+		{
+			// split up each Key=Value pair
+			FString AssetPath, Filepath;
+			if (AdditionalMountPoint.Split(TEXT("="), &AssetPath, &Filepath))
+			{
+				if (!AssetPath.IsEmpty() && !Filepath.IsEmpty())
+				{
+					if (FPaths::IsRelative(Filepath))
+					{
+						Filepath = FPaths::ConvertRelativePathToFull(FPaths::ProjectContentDir(), Filepath);
+					}
+
+					FStormSyncMountPointConfig Config;
+					Config.MountPoint = AssetPath;
+					Config.MountDirectory.Path = Filepath;
+					Configs.Add(MoveTemp(Config));
+				}
+			}
+		}
+	}
+	
+	return Configs;
+}
+
 void FStormSyncDrivesModule::OnPostEngineInit()
 {
 	UE_LOG(LogStormSyncDrives, Display, TEXT("FStormSyncDrivesModule::OnPostEngineInit ..."));
 	UE_LOG(LogStormSyncDrives, Display, TEXT("\t Mounting Drives based on config"));
 
 	ResetMountedDrivesFromSettings(GetDefault<UStormSyncDrivesSettings>());
+
+}
+
+namespace UE::StormSyncDrives::Private
+{
+TArray<FStormSyncMountPointConfig> RemoveDuplicatesAndConvertToAbsPaths(const TArray<FStormSyncMountPointConfig>& InMountPoints)
+{
+	TArray<FString> RootPaths;
+	TArray<FString> MountDirectories;
+
+	TArray<FStormSyncMountPointConfig> OutConfig;
+	for (FStormSyncMountPointConfig Entry : InMountPoints)
+	{
+		// Test for Mount Points duplicates
+		if (RootPaths.Contains(Entry.MountPoint) || MountDirectories.Contains(Entry.MountDirectory.Path))
+		{
+			continue;
+		}
+
+		if (FPaths::IsRelative(Entry.MountDirectory.Path))
+		{
+			Entry.MountDirectory.Path = FPaths::ConvertRelativePathToFull(FPaths::ProjectContentDir(), Entry.MountDirectory.Path);
+		}
+		
+		RootPaths.AddUnique(Entry.MountPoint);
+		MountDirectories.AddUnique(Entry.MountDirectory.Path);
+
+		OutConfig.Emplace(MoveTemp(Entry));
+	}
+	
+	return OutConfig;
 }
 
+void ConformSettingsToUserPreference(UStormSyncDrivesSettings* InSettings)
+{
+	bool bNeedsSave = false;
+	for (FStormSyncMountPointConfig& Entry : InSettings->MountPoints)
+	{
+		if (Entry.bResolveAsRelativePath && !FPaths::IsRelative(Entry.MountDirectory.Path))
+		{
+			FPaths::MakePathRelativeTo(Entry.MountDirectory.Path, *FPaths::ProjectContentDir());
+			bNeedsSave = true;
+		}
+		else if (!Entry.bResolveAsRelativePath && FPaths::IsRelative(Entry.MountDirectory.Path))
+		{
+			Entry.MountDirectory.Path = FPaths::ConvertRelativePathToFull(FPaths::ProjectContentDir(), Entry.MountDirectory.Path);
+			bNeedsSave = true;
+		}
+
+		if (Entry.MountPoint.EndsWith(TEXT("/")))
+		{
+			Entry.MountPoint = Entry.MountPoint.LeftChop(1);
+		}
+	}
+
+	if (bNeedsSave)
+	{
+		InSettings->SaveConfig();
+	}
+}
+
+} // namespace UE::StormSyncDrives::Private
+
+
 void FStormSyncDrivesModule::OnSettingsChanged(UObject* InSettings, FPropertyChangedEvent& InPropertyChangedEvent)
 {
 	const FName PropertyName = InPropertyChangedEvent.GetPropertyName();
@@ -135,79 +241,68 @@ void FStormSyncDrivesModule::OnSettingsChanged(UObject* InSettings, FPropertyCha
 
 	if (PropertyName == MountPointsPropertyName || MemberPropertyName == MountPointsPropertyName)
 	{
-		ResetMountedDrivesFromSettings(CastChecked<UStormSyncDrivesSettings>(InSettings));
+		UStormSyncDrivesSettings* Settings = CastChecked<UStormSyncDrivesSettings>(InSettings);
+		
+		UE::StormSyncDrives::Private::ConformSettingsToUserPreference(Settings);
+		ResetMountedDrivesFromSettings(Settings);
 	}
 }
 
-bool FStormSyncDrivesModule::ValidateMountPoint(const FStormSyncMountPointConfig& InMountPoint, int32 Index) const
+bool FStormSyncDrivesModule::IsMountPointInCache(const FStormSyncMountPointConfig& InMountPoint) const
+{
+	return CachedMountPoints.Contains(InMountPoint.MountPoint) &&
+		CachedMountPoints.FindChecked(InMountPoint.MountPoint).MountDirectory.Path == InMountPoint.MountDirectory.Path;
+}
+
+bool FStormSyncDrivesModule::MountDriveHelper(const FStormSyncMountPointConfig& InMountPoint)
 {
-	const FString RootPath = InMountPoint.MountPoint;
-	const FString Directory = InMountPoint.MountDirectory.Path;
+	if (CachedMountPoints.Contains(InMountPoint.MountPoint))
+	{
+		FText ErrorText;
+		if (!UnregisterMountPoint(InMountPoint, ErrorText))
+		{
+			UE_LOG(LogStormSyncDrives, Error, TEXT("FStormSyncDrivesModule::UnregisterMountPoint failed with Error: %s"), *ErrorText.ToString());
+			return false;
+		}
+	}
 
 	FText ErrorText;
-	if (!FStormSyncDrivesUtils::ValidateMountPoint(InMountPoint, ErrorText))
+	if (!RegisterMountPoint(InMountPoint, ErrorText))
 	{
-		const FText ErrorLog = FText::Format(
-			LOCTEXT("ValidationError_MountPointIndex", "Validation failed for MountPoint at index {0} (MountPoint: {1}, MountDirectory: {2}) - {3}"),
-			FText::AsNumber(Index),
-			FText::FromString(RootPath),
-			FText::FromString(Directory),
-			ErrorText
-		);
-		AddMessageError(ErrorLog);
-		UE_LOG(LogStormSyncDrives, Error, TEXT("FStormSyncDrivesModule::ValidateMountPoint - %s"), *ErrorLog.ToString());
+		UE_LOG(LogStormSyncDrives, Error, TEXT("FStormSyncDrivesModule::RegisterMountPoint failed with Error: %s"), *ErrorText.ToString());
 		return false;
 	}
 
 	return true;
 }
 
-bool FStormSyncDrivesModule::ValidateMountPoints(const TArray<FStormSyncMountPointConfig>& InMountPoints) const
+bool FStormSyncDrivesModule::RemoveFromCacheIfNotInConfigs(const TArray<FStormSyncMountPointConfig>& Configs)
 {
-#if WITH_EDITOR
-	// Reset log listing and clear message on each validation run
-	if (LogListing.IsValid())
+	TArray<FStormSyncMountPointConfig> UnregisterList;
+	
+	// Collect all configs that should be unmounted. 
+	for (const TPair<FString, FStormSyncMountPointConfig>& Cache : CachedMountPoints)
 	{
-		LogListing->ClearMessages();
-	}
-#endif
-
-	bool bResult = true;
-
-	// First check for mount point validity individually
-	int32 Index = 0;
-	for (const FStormSyncMountPointConfig& Entry : InMountPoints)
-	{
-		UE_LOG(LogStormSyncDrives, Display, TEXT("Try validate entry %s %s"), *Entry.MountPoint, *Entry.MountDirectory.Path);
-
-		// First check for mount point validity individually
-		bResult &= ValidateMountPoint(Entry, Index);
-		Index++;
-	}
-
-	// Then check for duplicates
-	TArray<FText> ValidationErrors;
-	if (!FStormSyncDrivesUtils::ValidateNonDuplicates(InMountPoints, ValidationErrors))
-	{
-		bResult = false;
-
-		// Log any returned errors to message log
-		for (const FText& ValidationError : ValidationErrors)
+		bool bInSettings = Algo::AnyOf(Configs, [&Key = Cache.Key](const FStormSyncMountPointConfig& InMount){ return InMount.MountPoint == Key; });
+		if (!bInSettings)
 		{
-			AddMessageError(ValidationError);
-			UE_LOG(LogStormSyncDrives, Error, TEXT("FStormSyncDrivesModule::ValidateMountPoints - %s"), *ValidationError.ToString());
+			UnregisterList.Add(Cache.Value);
 		}
 	}
 
-#if WITH_EDITOR
-	// Got some errors, notify user
-	if (!bResult)
+	// Now try to unmount.
+	bool bSuccess = true;
+	for (const FStormSyncMountPointConfig& Config : UnregisterList)
 	{
-		LogListing->NotifyIfAnyMessages(LOCTEXT("Validation_Error", "Storm Sync: Error validating Mount Points settings"), EMessageSeverity::Error, true);
+		FText ErrorText;
+		if (!UnregisterMountPoint(Config, ErrorText))
+		{
+			UE_LOG(LogStormSyncDrives, Error, TEXT("FStormSyncDrivesModule::RemoveFromCacheIfNotInConfigs failed with Error: %s"), *ErrorText.ToString());
+			bSuccess = false;
+		}
 	}
-#endif
 
-	return bResult;
+	return bSuccess;
 }
 
 void FStormSyncDrivesModule::ResetMountedDrivesFromSettings(const UStormSyncDrivesSettings* InSettings)
@@ -215,57 +310,47 @@ void FStormSyncDrivesModule::ResetMountedDrivesFromSettings(const UStormSyncDriv
 	UE_LOG(LogStormSyncDrives, Display, TEXT("FStormSyncDrivesModule::ResetMountedDrives ..."));
 	check(InSettings);
 
-	if (ValidateMountPoints(InSettings->MountPoints))
-	{
-		UnregisterMountedDrives();
-		CacheMountedDrives(InSettings->MountPoints);
-		RegisterMountedDrives();
-	}
-}
+	// Collect any command line configs the user has specified. 
+	TArray<FStormSyncMountPointConfig> Configs = GetCommandLineConfigs();
+	Configs += InSettings->MountPoints;
 
-void FStormSyncDrivesModule::UnregisterMountedDrives()
-{
-	UE_LOG(LogStormSyncDrives, Display, TEXT("FStormSyncDrivesModule::UnregisterMountedDrives ..."));
-	for (const TSharedPtr<FStormSyncMountPointConfig>& MountedDrive : MountedDrives)
-	{
-		if (!MountedDrive.IsValid())
-		{
-			continue;
-		}
+	TArray<FStormSyncMountPointConfig> Filtered = UE::StormSyncDrives::Private::RemoveDuplicatesAndConvertToAbsPaths(Configs);
+	UE_CLOG(Filtered.Num() != Configs.Num(), LogStormSyncDrives, Warning, TEXT("FStormSyncDrivesModule::ResetMountedDrives settings object contains duplicate values."));
 
-		FText ErrorText;
-		if (!UnregisterMountPoint(*MountedDrive.Get(), ErrorText))
+#if WITH_EDITOR
+	LogListing->ClearMessages();
+#endif
+
+	// Mount any config not in our current cache. 
+	bool bSuccess = true;
+	for (const FStormSyncMountPointConfig& MountPoint : Filtered)
+	{
+		if (!IsMountPointInCache(MountPoint))
 		{
-			UE_LOG(LogStormSyncDrives, Error, TEXT("FStormSyncDrivesModule::UnregisterMountedDrives failed with Error: %s"), *ErrorText.ToString());
+			bSuccess &= MountDriveHelper(MountPoint);
 		}
 	}
-}
 
-void FStormSyncDrivesModule::CacheMountedDrives(const TArray<FStormSyncMountPointConfig>& InMountPoints)
-{
-	UE_LOG(LogStormSyncDrives, Display, TEXT("FStormSyncDrivesModule::CacheMountedDrives ..."));
-	MountedDrives.Reset();
+	// Remove any stale cache items. 
+	bSuccess &= RemoveFromCacheIfNotInConfigs(Filtered);
 
-	for (const FStormSyncMountPointConfig& Entry : InMountPoints)
-	{
-		MountedDrives.Add(MakeShared<FStormSyncMountPointConfig>(Entry));
-	}
+#if WITH_EDITOR
+	if (!bSuccess)
+ 	{
+		LogListing->NotifyIfAnyMessages(LOCTEXT("Validation_Error", "Storm Sync: Error validating Mount Points settings"), EMessageSeverity::Error, true);
+ 	}
+#endif
 }
 
-void FStormSyncDrivesModule::RegisterMountedDrives()
+void FStormSyncDrivesModule::UnregisterMountedDrives()
 {
-	UE_LOG(LogStormSyncDrives, Display, TEXT("FStormSyncDrivesModule::RegisterMountedDrives ..."));
-	for (TSharedPtr<FStormSyncMountPointConfig> MountedDrive : MountedDrives)
+	UE_LOG(LogStormSyncDrives, Display, TEXT("FStormSyncDrivesModule::UnregisterMountedDrives ..."));
+	for (const TPair<FString, FStormSyncMountPointConfig>& Cache : CachedMountPoints)
 	{
-		if (!MountedDrive.IsValid())
-		{
-			continue;
-		}
-
 		FText ErrorText;
-		if (!RegisterMountPoint(*MountedDrive.Get(), ErrorText))
+		if (!UnregisterMountPoint(Cache.Value, ErrorText))
 		{
-			UE_LOG(LogStormSyncDrives, Error, TEXT("FStormSyncDrivesModule::RegisterMountedDrives failed with Error: %s"), *ErrorText.ToString());
+			UE_LOG(LogStormSyncDrives, Error, TEXT("FStormSyncDrivesModule::UnregisterMountedDrives failed with Error: %s"), *ErrorText.ToString());
 		}
 	}
 }
diff --git a/Engine/Plugins/VirtualProduction/StormSync/Source/StormSyncDrives/Private/StormSyncDrivesModule.h b/Engine/Plugins/VirtualProduction/StormSync/Source/StormSyncDrives/Private/StormSyncDrivesModule.h
--- a/Engine/Plugins/VirtualProduction/StormSync/Source/StormSyncDrives/Private/StormSyncDrivesModule.h
+++ b/Engine/Plugins/VirtualProduction/StormSync/Source/StormSyncDrives/Private/StormSyncDrivesModule.h
@@ -5,6 +5,8 @@
 #include "IStormSyncDrivesModule.h"
 #include "UObject/UnrealType.h"
 
+#include "StormSyncDrivesSettings.h"
+
 #if WITH_EDITOR
 #include "IMessageLogListing.h"
 #endif
@@ -29,13 +31,12 @@ public:
 private:
 	/** The name of the message output log we will send messages to */
 	static constexpr const TCHAR* LogName = TEXT("StormSyncDrives");
-	
+
 	/**
-	 * Array of Mount Points templates populated from UStormSyncDrivesSettings.
-	 * 
-	 * Allows projects to specify reusable mount points templates for the content browser.
-	*/
-	TArray<TSharedPtr<FStormSyncMountPointConfig>> MountedDrives;
+	 * Stores a cache of the currently mounted paths. This will allow us to determine when the user changes settings
+	 * and require us to update our mount mapping. 
+	 */
+	TMap<FString, FStormSyncMountPointConfig> CachedMountPoints;
 
 #if WITH_EDITORONLY_DATA
 	/** Listing used in the editor by the Message Log. */
@@ -45,6 +46,12 @@ private:
 	/** Core OnPostEngineInit delegate used to mount any stored config in developer settings */
 	void OnPostEngineInit();
 
+	/** Return the storm sync mount point config that was set on the command line. */
+	TArray<FStormSyncMountPointConfig> GetCommandLineConfigs() const;
+
+	/** If the cache configs are not in the given list then unregister and remove them from the cache. */
+	bool RemoveFromCacheIfNotInConfigs(const TArray<FStormSyncMountPointConfig>& Configs);
+
 	/**
 	 * Developer settings changed handler. Called anytime settings changed.
 	 *
@@ -52,15 +59,16 @@ private:
 	 */
 	void OnSettingsChanged(UObject* InSettings, FPropertyChangedEvent& InPropertyChangedEvent);
 
-	/** Performs validation for a single mount point config and log any errors to the message log */
-	bool ValidateMountPoint(const FStormSyncMountPointConfig& InMountPoint, int32 Index) const;
-	
 	/**
-	 * Performs validation for the provided list of mount point config. Logs any errors to the message log.
-	 *
-	 * Happens once on engine init, and anytime the user settings changed.
+	 * Checks to see if the given mount point is currently in the cache.  If yes, then that implies that
+	 * that we have already mapped the mount point and we do not need to re-process.
+	 */
+	bool IsMountPointInCache(const FStormSyncMountPointConfig& InMountPoint) const;
+
+	/**
+	 * Helper function that unregisters an existing mount point and remaps it.
 	 */
-	bool ValidateMountPoints(const TArray<FStormSyncMountPointConfig>& InMountPoints) const;
+	bool MountDriveHelper(const FStormSyncMountPointConfig& InMountPoint);
 
 	/**
 	 * Unregisters previously mounted drives, cache (update internal list of MountedDrives shared pointers),
@@ -70,15 +78,8 @@ private:
 	 */
 	void ResetMountedDrivesFromSettings(const UStormSyncDrivesSettings* InSettings);
 
-	/** Unregisters any previously mounted drives */
 	void UnregisterMountedDrives();
-
-	/** Reset and updates internal list of registered drives */
-	void CacheMountedDrives(const TArray<FStormSyncMountPointConfig>& InMountPoints);
-
-	/** Tries to mount previously cached list of mount point config */
-	void RegisterMountedDrives();
-
+	
 	/** Static helper to log an error to the message log. Used during validation process. */
 	static void AddMessageError(const FText& Text);
 };
diff --git a/Engine/Plugins/VirtualProduction/StormSync/Source/StormSyncDrives/Public/StormSyncDrivesSettings.h b/Engine/Plugins/VirtualProduction/StormSync/Source/StormSyncDrives/Public/StormSyncDrivesSettings.h
--- a/Engine/Plugins/VirtualProduction/StormSync/Source/StormSyncDrives/Public/StormSyncDrivesSettings.h
+++ b/Engine/Plugins/VirtualProduction/StormSync/Source/StormSyncDrives/Public/StormSyncDrivesSettings.h
@@ -19,16 +19,22 @@ struct FStormSyncMountPointConfig
 	UPROPERTY(EditAnywhere, Category="Mounted Drives")
 	FDirectoryPath MountDirectory;
 
+	/** The path in MountDirectory should be interpreted as a relative path to the content directory of the active project. */
+	UPROPERTY(EditAnywhere, Category = "Mounted Drives")
+	bool bResolveAsRelativePath = false;
+	
 	FStormSyncMountPointConfig() = default;
 
-	FStormSyncMountPointConfig(const FString& InMountPoint, const FDirectoryPath& InMountDirectory)
+	FStormSyncMountPointConfig(const FString& InMountPoint, const FDirectoryPath& InMountDirectory, bool bInResolveAsRelativePath = false)
 		: MountPoint(InMountPoint)
 		, MountDirectory(InMountDirectory)
+		, bResolveAsRelativePath(bInResolveAsRelativePath)
 	{
 	}
 
-	FStormSyncMountPointConfig(const FString& InMountPoint, const FString& InMountDirectory)
+	FStormSyncMountPointConfig(const FString& InMountPoint, const FString& InMountDirectory, bool bInResolveAsRelativePath = false)
 		: MountPoint(InMountPoint)
+		, bResolveAsRelativePath(bInResolveAsRelativePath)
 	{
 		MountDirectory.Path = InMountDirectory;
 	}
