Storm Sync: fix crash when using context menu on class items in content browser

diff --git a/Engine/Plugins/VirtualProduction/StormSync/Source/StormSyncEditor/Private/ContextMenu/StormSyncAssetFolderContextMenu.cpp b/Engine/Plugins/VirtualProduction/StormSync/Source/StormSyncEditor/Private/ContextMenu/StormSyncAssetFolderContextMenu.cpp
--- a/Engine/Plugins/VirtualProduction/StormSync/Source/StormSyncEditor/Private/ContextMenu/StormSyncAssetFolderContextMenu.cpp
+++ b/Engine/Plugins/VirtualProduction/StormSync/Source/StormSyncEditor/Private/ContextMenu/StormSyncAssetFolderContextMenu.cpp
@@ -8,6 +8,7 @@
 #include "ContentBrowserAssetDataCore.h"
 #include "ContentBrowserAssetDataPayload.h"
 #include "ContentBrowserDataMenuContexts.h"
+#include "ContentBrowserDataSource.h"
 #include "Engine/BlueprintGeneratedClass.h"
 #include "HAL/FileManager.h"
 #include "IStormSyncTransportClientModule.h"
@@ -455,21 +456,23 @@ void FStormSyncAssetFolderContextMenu::GetSelectedFilesAndFolders(const UToolMen
 			continue;
 		}
 
-		const UContentBrowserDataSource* DataSource = SelectedItemData->GetOwnerDataSource();
+		UContentBrowserDataSource* const DataSource = SelectedItemData->GetOwnerDataSource();
 		if (!DataSource)
 		{
 			continue;
 		}
 
-		for (const FContentBrowserItemData& InternalItems : SelectedItem.GetInternalItems())
+		for (const FContentBrowserItemData& Item : SelectedItem.GetInternalItems())
 		{
-			if (const TSharedPtr<const FContentBrowserAssetFolderItemDataPayload> FolderPayload = ContentBrowserAssetData::GetAssetFolderItemPayload(DataSource, InternalItems))
+			FName PackagePath;
+			FAssetData AssetData;
+			if (DataSource->Legacy_TryGetPackagePath(Item, PackagePath))
 			{
-				OutSelectedPackagePaths.Add(FolderPayload->GetInternalPath().ToString());
+				OutSelectedPackagePaths.Add(PackagePath.ToString());
 			}
-			else if (const TSharedPtr<const FContentBrowserAssetFileItemDataPayload> ItemPayload = ContentBrowserAssetData::GetAssetFileItemPayload(DataSource, InternalItems))
+			else if (DataSource->Legacy_TryGetAssetData(Item, AssetData))
 			{
-				OutSelectedAssetPackages.Add(ItemPayload->GetAssetData().PackageName.ToString());
+				OutSelectedAssetPackages.Add(AssetData.PackageName.ToString());
 			}
 		}
 	}
@@ -492,17 +495,18 @@ void FStormSyncAssetFolderContextMenu::GetSelectedFiles(const UToolMenu* InMenu,
 			continue;
 		}
 
-		const UContentBrowserDataSource* DataSource = SelectedItemData->GetOwnerDataSource();
+		UContentBrowserDataSource* const DataSource = SelectedItemData->GetOwnerDataSource();
 		if (!DataSource)
 		{
 			continue;
 		}
 
-		for (const FContentBrowserItemData& InternalItems : SelectedItem.GetInternalItems())
+		for (const FContentBrowserItemData& Item : SelectedItem.GetInternalItems())
 		{
-			if (const TSharedPtr<const FContentBrowserAssetFileItemDataPayload> ItemPayload = ContentBrowserAssetData::GetAssetFileItemPayload(DataSource, InternalItems))
+			FAssetData AssetData;
+			if (DataSource->Legacy_TryGetAssetData(Item, AssetData))
 			{
-				OutSelectedAssetPackages.Add(ItemPayload->GetAssetData().PackageName.ToString());
+				OutSelectedAssetPackages.Add(AssetData.PackageName.ToString());
 			}
 		}
 	}
