[StormSync] Fix Tcp Port appended when enumerating local adapters.

Problem:
The bug was that the enumerated local addresses had the wrong port, using the configured port instead of the actual port. Because of port collision avoidance, when launching multiple servers on the same host, each server will get a different port then what it was configured.

Context:
When a request (pull, push, sync) is received on the server, the local adapters are enumerated and sent back as a response to the client. The client may then use those addresses to connect to the server.
The messages have both HostAddress and a list of adapters. If the host address is an Any (0.0.0.0) address (which is the default), then it will try each adapter to connect to the server.

#Synced-CL 47720872

diff --git a/Engine/Plugins/VirtualProduction/StormSync/Source/StormSyncTransportCore/Private/StormSyncTransportNetworkUtils.cpp b/Engine/Plugins/VirtualProduction/StormSync/Source/StormSyncTransportCore/Private/StormSyncTransportNetworkUtils.cpp
--- a/Engine/Plugins/VirtualProduction/StormSync/Source/StormSyncTransportCore/Private/StormSyncTransportNetworkUtils.cpp
+++ b/Engine/Plugins/VirtualProduction/StormSync/Source/StormSyncTransportCore/Private/StormSyncTransportNetworkUtils.cpp
@@ -45,11 +45,20 @@ FString FStormSyncTransportNetworkUtils::GetTcpEndpointAddress()
 TArray<FString> FStormSyncTransportNetworkUtils::GetLocalAdapterAddresses()
 {
 	TArray<FString> AdapterAddresses;
-	
-	const UStormSyncTransportSettings* Settings = GetDefault<UStormSyncTransportSettings>();
-	check(Settings);
-	
-	const FString ServerEndpoint = Settings->GetServerEndpoint();
+
+	// Request the current server endpoint address in order to extract it's current port,
+	// which may be different from the configured port.
+	FString ServerEndpoint = GetCurrentTcpServerEndpointAddress();
+
+	// If the server is not currently running, fallback to the settings.
+	if (ServerEndpoint.IsEmpty())
+	{
+		const UStormSyncTransportSettings* Settings = GetDefault<UStormSyncTransportSettings>();
+		if (ensure(Settings))
+		{
+			ServerEndpoint = Settings->GetServerEndpoint();
+		}
+	}
 	
 	FIPv4Endpoint Endpoint;
 	if (!FIPv4Endpoint::Parse(ServerEndpoint, Endpoint))
diff --git a/Engine/Plugins/VirtualProduction/StormSync/Source/StormSyncTransportCore/Public/StormSyncTransportNetworkUtils.h b/Engine/Plugins/VirtualProduction/StormSync/Source/StormSyncTransportCore/Public/StormSyncTransportNetworkUtils.h
--- a/Engine/Plugins/VirtualProduction/StormSync/Source/StormSyncTransportCore/Public/StormSyncTransportNetworkUtils.h
+++ b/Engine/Plugins/VirtualProduction/StormSync/Source/StormSyncTransportCore/Public/StormSyncTransportNetworkUtils.h
@@ -22,14 +22,14 @@ public:
 	static FString GetTcpEndpointAddress();
 	
 	/**
-	 * Returns the list of local adapter addresses as returned by the socket subsystem, with port from server
-	 * endpoint settings appended to the address list.
+	 * Returns the list of local adapter addresses as returned by the socket subsystem, with current port from server
+	 * endpoint appended to the address list.
 	 */
 	static TArray<FString> GetLocalAdapterAddresses();
 
 	/**
 	 * Returns the current tcp server endpoint address (ip:port).
-	 * @remark This is potentially different than the configured endpoint address in case of port collisions, it may be on a different one.
+	 * @remark This is potentially different from the configured endpoint address in case of port collisions, it may be on a different one.
 	 */
 	static FString GetCurrentTcpServerEndpointAddress();
 
