Storm Sync: fixed ensure fail by using FApp::GetCurrentTime in a different thread via FRunnable.
Using FPlatform::Seconds in the meantime (at the expense of higher cost)

#Synced-CL 48355279

diff --git a/Engine/Plugins/VirtualProduction/StormSync/Source/StormSyncTransportServer/Private/ServiceDiscovery/StormSyncDiscoveryManager.cpp b/Engine/Plugins/VirtualProduction/StormSync/Source/StormSyncTransportServer/Private/ServiceDiscovery/StormSyncDiscoveryManager.cpp
--- a/Engine/Plugins/VirtualProduction/StormSync/Source/StormSyncTransportServer/Private/ServiceDiscovery/StormSyncDiscoveryManager.cpp
+++ b/Engine/Plugins/VirtualProduction/StormSync/Source/StormSyncTransportServer/Private/ServiceDiscovery/StormSyncDiscoveryManager.cpp
@@ -63,8 +63,8 @@ uint32 FStormSyncDiscoveryManager::Run()
 {
 	while (bRunning)
 	{
-		const double CurrentTime = FApp::GetCurrentTime();
-		
+		const double CurrentTime = FPlatformTime::Seconds();
+
 		TArray<FStormSyncDelegateItem> DelegateQueue;
 		{
 			FScopeLock Lock(&ConnectionLastActiveSection);
@@ -162,7 +162,7 @@ void FStormSyncDiscoveryManager::PublishConnectMessage()
 		Message->StormSyncServerAddressId = FStormSyncTransportNetworkUtils::GetServerEndpointMessageAddress();
 		Message->StormSyncClientAddressId = FStormSyncTransportNetworkUtils::GetClientEndpointMessageAddress();
 
-		LastPublishTime = FApp::GetCurrentTime();
+		LastPublishTime = FPlatformTime::Seconds();
 		MessageEndpoint->Publish(Message.Release());
 	}
 }
